- Adding Rivermax input stream management to handle reception of video frame buffers over IP
- Adding Rivermax output stream management to handle sending video frame buffers over IP

diff --git a/Engine/Plugins/VirtualProduction/Rivermax/RivermaxCore/RivermaxCore.uplugin b/Engine/Plugins/VirtualProduction/Rivermax/RivermaxCore/RivermaxCore.uplugin
new file mode 100644
--- /dev/null
+++ b/Engine/Plugins/VirtualProduction/Rivermax/RivermaxCore/RivermaxCore.uplugin
@@ -0,0 +1,32 @@
+{
+	"FileVersion": 3,
+	"Version": 1,
+	"VersionName": "1.0",
+	"FriendlyName": "RivermaxCore",
+	"Description": "Base plugin exposing rivermax to engine",
+	"Category": "Virtual Production",
+	"CreatedBy": "Epic Games Inc",
+	"CreatedByURL": "http://epicgames.com",
+	"DocsURL": "",
+	"MarketplaceURL": "",
+	"SupportURL": "",
+	"Hidden": "true",
+	"CanContainContent": false,
+	"IsBetaVersion": true,
+	"Installed": false,
+	"SupportedTargetPlatforms": [ "Win64" ],
+	"Plugins": [
+		{
+			"Name": "MediaIOFramework",
+			"Enabled": true
+		}
+	],
+	"Modules": [
+		{
+			"Name": "RivermaxCore",
+			"Type": "Runtime",
+			"LoadingPhase": "Default",
+			"PlatformAllowList": [ "Win64" ]
+		}
+	]
+}
\ No newline at end of file
diff --git a/Engine/Plugins/VirtualProduction/Rivermax/RivermaxCore/Source/Private/RivermaxCoreModule.cpp b/Engine/Plugins/VirtualProduction/Rivermax/RivermaxCore/Source/Private/RivermaxCoreModule.cpp
new file mode 100644
--- /dev/null
+++ b/Engine/Plugins/VirtualProduction/Rivermax/RivermaxCore/Source/Private/RivermaxCoreModule.cpp
@@ -0,0 +1,38 @@
+// Copyright Epic Games, Inc. All Rights Reserved.
+
+#include "RivermaxCoreModule.h"
+
+#include "Modules/ModuleManager.h"
+#include "RivermaxInputStream.h"
+#include "RivermaxManager.h"
+#include "RivermaxOutputStream.h"
+
+
+void FRivermaxCoreModule::StartupModule()
+{
+	RivermaxManager = MakeShared<UE::RivermaxCore::Private::FRivermaxManager>();
+}
+
+void FRivermaxCoreModule::ShutdownModule()
+{
+
+}
+
+TUniquePtr<UE::RivermaxCore::IRivermaxInputStream> FRivermaxCoreModule::CreateInputStream()
+{
+	using UE::RivermaxCore::Private::FRivermaxInputStream;
+	return MakeUnique<FRivermaxInputStream>();
+}
+
+TUniquePtr<UE::RivermaxCore::IRivermaxOutputStream> FRivermaxCoreModule::CreateOutputStream()
+{
+	using UE::RivermaxCore::Private::FRivermaxOutputStream;
+	return MakeUnique<FRivermaxOutputStream>();
+}
+
+TSharedPtr<UE::RivermaxCore::IRivermaxManager> FRivermaxCoreModule::GetRivermaxManager()
+{
+	return RivermaxManager;
+}
+
+IMPLEMENT_MODULE(FRivermaxCoreModule, RivermaxCore);
diff --git a/Engine/Plugins/VirtualProduction/Rivermax/RivermaxCore/Source/Private/RivermaxCoreModule.h b/Engine/Plugins/VirtualProduction/Rivermax/RivermaxCore/Source/Private/RivermaxCoreModule.h
new file mode 100644
--- /dev/null
+++ b/Engine/Plugins/VirtualProduction/Rivermax/RivermaxCore/Source/Private/RivermaxCoreModule.h
@@ -0,0 +1,35 @@
+// Copyright Epic Games, Inc. All Rights Reserved.
+
+#pragma once
+
+#include "IRivermaxCoreModule.h"
+
+
+namespace UE::RivermaxCore
+{
+	class IRivermaxManager;
+}
+
+/**
+ * Implementation of the IRivermaxCoreModule managing Rivermax library and stream creation
+ */
+class FRivermaxCoreModule : public IRivermaxCoreModule
+{
+
+public:
+	//~ Begin IModuleInterface
+	virtual void StartupModule() override;
+	virtual void ShutdownModule() override;
+	//~ End IModuleInterface
+
+	//~ Begin IRivermaxCoreModule interface
+	virtual TUniquePtr<UE::RivermaxCore::IRivermaxInputStream> CreateInputStream() override;
+	virtual TUniquePtr<UE::RivermaxCore::IRivermaxOutputStream> CreateOutputStream() override;
+	virtual TSharedPtr<UE::RivermaxCore::IRivermaxManager> GetRivermaxManager() override;
+	//~ End IRivermaxCoreModule interface
+
+private:
+
+	TSharedPtr<UE::RivermaxCore::IRivermaxManager> RivermaxManager;
+};
+
diff --git a/Engine/Plugins/VirtualProduction/Rivermax/RivermaxCore/Source/Private/RivermaxInputStream.cpp b/Engine/Plugins/VirtualProduction/Rivermax/RivermaxCore/Source/Private/RivermaxInputStream.cpp
new file mode 100644
--- /dev/null
+++ b/Engine/Plugins/VirtualProduction/Rivermax/RivermaxCore/Source/Private/RivermaxInputStream.cpp
@@ -0,0 +1,429 @@
+// Copyright Epic Games, Inc. All Rights Reserved.
+
+#include "RivermaxInputStream.h"
+
+#include "Async/Async.h"
+#include "IRivermaxCoreModule.h"
+#include "IRivermaxManager.h"
+#include "Misc/ByteSwap.h"
+#include "RivermaxLog.h"
+#include "RivermaxUtils.h"
+
+#if PLATFORM_WINDOWS
+#include <WS2tcpip.h>
+#endif
+
+
+namespace UE::RivermaxCore::Private
+{
+#if PLATFORM_SUPPORTS_PRAGMA_PACK
+#pragma pack(push, 1)
+#endif
+	struct FRivermaxRTPSampleRowData
+	{
+		uint32 ContributingSourceCount : 4;
+		uint32 ExtensionBit : 1;
+		uint32 PaddingBit : 1;
+		uint32 Version : 2;
+		uint32 PayloadType : 7;
+		uint32 MarkerBit : 1;
+		uint32 SequenceNumber : 16;
+		uint32 Timestamp : 32;
+		uint32 SynchronizationSource : 32;
+		uint32 ExtendedSequenceNumber : 16;
+		//SRD 1
+		uint32 SRDLength1 : 16;
+		uint32 SRDRowNumberHigh1 : 7;
+		uint32 FieldIdentification1 : 1;
+		uint32 SRDRowNumberLow1 : 8;
+		uint32 SRDOffsetHigh1 : 7;
+		uint32 ContinuationBit1 : 1;
+		uint32 SRDOffsetLow1 : 8;
+		//SRD 2
+		uint32 SRDLength2 : 16;
+		uint32 SRDRowNumberHigh2 : 7;
+		uint32 FieldIdentification2 : 1;
+		uint32 SRDRowNumberLow2 : 8;
+		uint32 SRDOffsetHigh2 : 7;
+		uint32 ContinuationBit2 : 1;
+		uint32 SRDOffsetLow2 : 8;
+
+		uint16 GetSrd1RowNumber() { return ((SRDRowNumberHigh1 << 8) | SRDRowNumberLow1); }
+		uint16 GetSrd1Offset() { return ((SRDOffsetHigh1 << 8) | SRDOffsetLow1); }
+
+		uint16 GetSrd2RowNumber() { return ((SRDRowNumberHigh2 << 8) | SRDRowNumberLow2); }
+		uint16 GetSrd2Offset() { return ((SRDOffsetHigh2 << 8) | SRDOffsetLow2); }
+	};
+#if PLATFORM_SUPPORTS_PRAGMA_PACK
+#pragma pack(pop)
+#endif
+
+
+	uint8* GetRTPHeaderPointer(uint8* InHeader)
+	{
+		if (!InHeader /*|| RMAX_APP_PROTOCOL_PACKET == m_type*/)
+		{
+			return nullptr;
+		}
+
+		static constexpr uint32 ETH_TYPE_802_1Q = 0x8100;          /* 802.1Q VLAN Extended Header  */
+		static constexpr uint32 RTP_HEADER_SIZE = 12;
+		uint16* ETHProto = (uint16_t*)(InHeader + RTP_HEADER_SIZE);
+		if (ETH_TYPE_802_1Q == ByteSwap(*ETHProto))
+		{
+			InHeader += 46; // 802 + 802.1Q + IP + UDP
+		}
+		else
+		{
+			InHeader += 42; // 802 + IP + UDP
+		}
+		return InHeader;
+	}
+
+	FRivermaxInputStream::FRivermaxInputStream()
+	{
+
+	}
+
+	FRivermaxInputStream::~FRivermaxInputStream()
+	{
+		Uninitialize();
+	}
+
+	bool FRivermaxInputStream::Initialize(const FRivermaxStreamOptions& InOptions, IRivermaxInputStreamListener& InListener)
+	{
+		IRivermaxCoreModule& RivermaxModule = FModuleManager::LoadModuleChecked<IRivermaxCoreModule>(TEXT("RivermaxCore"));
+		if (RivermaxModule.GetRivermaxManager()->IsInitialized() == false)
+		{
+			UE_LOG(LogRivermax, Warning, TEXT("Can't create Rivermax Input Stream. Library isn't initialized."));
+			return false;
+		}
+
+		Options = InOptions;
+		Listener = &InListener;
+
+		Async(EAsyncExecution::TaskGraph, [this]()
+		{
+			bool bWasSuccessful = false;
+			int32 FlowId = 0; //todo configure
+
+			//Configure local IP interface
+			const rmax_in_stream_type StreamType = RMAX_RAW_PACKET;
+			sockaddr_in RivermaxInterface;
+			memset(&RivermaxInterface, 0, sizeof(RivermaxInterface));
+			if (inet_pton(AF_INET, StringCast<ANSICHAR>(*Options.SourceAddress).Get(), &RivermaxInterface.sin_addr) != 1)
+			{
+				UE_LOG(LogRivermax, Warning, TEXT("inet_pton failed to %s"), *Options.SourceAddress);
+			}
+			else
+			{
+				RivermaxInterface.sin_family = AF_INET;
+
+				//Configure Flow and destination IP (multicast)
+				memset(&FlowAttribute, 0, sizeof(FlowAttribute));
+				FlowAttribute.local_addr.sin_family = AF_INET;
+				FlowAttribute.flow_id = FlowId;
+				if (inet_pton(AF_INET, StringCast<ANSICHAR>(*Options.DestinationAddress).Get(), &FlowAttribute.local_addr.sin_addr) != 1)
+				{
+					UE_LOG(LogRivermax, Warning, TEXT("inet_pton failed to %s"), *Options.DestinationAddress);
+				}
+				else
+				{
+					FlowAttribute.local_addr.sin_port = ByteSwap((uint16)Options.Port);
+
+					const rmax_in_buffer_attr_flags_t BufferAttributeFlags = RMAX_IN_BUFFER_ATTER_FLAG_NONE; //todo whether ordering is based on sequence or extended sequence
+					uint32 BufferElement = 1 << 18;//todo number of packets to allocate memory for
+					rmax_in_buffer_attr BufferAttributes;
+					FMemory::Memset(&BufferAttributes, 0, sizeof(BufferAttributes));
+					BufferAttributes.num_of_elements = BufferElement;
+					BufferAttributes.attr_flags = BufferAttributeFlags;
+
+					FMemory::Memset(&BufferConfiguration.DataMemory, 0, sizeof(BufferConfiguration.DataMemory));
+					BufferConfiguration.DataMemory.max_size = BufferConfiguration.DataMemory.min_size = BufferConfiguration.PayloadExpectedSize;
+					BufferAttributes.data = &BufferConfiguration.DataMemory;
+
+					FMemory::Memset(&BufferConfiguration.HeaderMemory, 0, sizeof(BufferConfiguration.HeaderMemory));
+					BufferConfiguration.HeaderMemory.max_size = BufferConfiguration.HeaderMemory.min_size = BufferConfiguration.HeaderExpectedSize;
+					BufferAttributes.hdr = &BufferConfiguration.HeaderMemory;
+
+					rmax_status_t Status = rmax_in_query_buffer_size(StreamType, &RivermaxInterface, &BufferAttributes, &BufferConfiguration.PayloadSize, &BufferConfiguration.HeaderSize);
+					if (Status == RMAX_OK)
+					{
+						const uint32 CacheLineSize = PLATFORM_CACHE_LINE_SIZE;
+
+						BufferConfiguration.DataMemory.ptr = FMemory::Malloc(BufferConfiguration.PayloadSize, CacheLineSize);
+						BufferConfiguration.HeaderMemory.ptr = FMemory::Malloc(BufferConfiguration.HeaderSize, CacheLineSize);
+
+						//Buffers configured, now configure stream and attach flow
+						const rmax_in_timestamp_format TimestampFormat = rmax_in_timestamp_format::RMAX_PACKET_TIMESTAMP_RAW_NANO; //how packets are stamped. counter or nanoseconds
+						const rmax_in_flags InputFlags = rmax_in_flags::RMAX_IN_CREATE_STREAM_INFO_PER_PACKET; //default value for 2110 in example
+						Status = rmax_in_create_stream(StreamType, &RivermaxInterface, &BufferAttributes, TimestampFormat, InputFlags, &StreamId);
+						if (Status == RMAX_OK)
+						{
+							Status = rmax_in_attach_flow(StreamId, &FlowAttribute);
+							if (Status == RMAX_OK)
+							{
+								bIsActive = true;
+								RivermaxThread = FRunnableThread::Create(this, TEXT("Rivermax InputStream Thread"), 128 * 1024, TPri_AboveNormal, FPlatformAffinity::GetPoolThreadMask());
+								bWasSuccessful = true;
+							}
+							else
+							{
+								UE_LOG(LogRivermax, Warning, TEXT("Could not attach flow to stream. Status: %d."), Status);
+							}
+						}
+						else
+						{
+							UE_LOG(LogRivermax, Warning, TEXT("Could not create stream. Status: %d."), Status);
+						}
+					}
+					else
+					{
+						UE_LOG(LogRivermax, Warning, TEXT("Could not query buffer size. Status: %d"), Status);
+					}
+				}
+			}
+
+			Listener->OnInitializationCompleted(bWasSuccessful);
+		});
+		
+		return true;
+	}
+
+	void FRivermaxInputStream::Uninitialize()
+	{
+		if (RivermaxThread != nullptr)
+		{
+			Stop();
+			RivermaxThread->Kill(true);
+			delete RivermaxThread;
+			RivermaxThread = nullptr;
+			UE_LOG(LogRivermax, Log, TEXT("Rivermax Input stream has shutdown"));
+		}
+	}
+
+	void FRivermaxInputStream::Process_AnyThread()
+	{
+		const size_t MinChunkSize = 0;
+		const size_t MaxChunkSize = 5000;
+		const int Timeout = 0;
+		const int Flags = 0;
+		rmax_in_completion Completion;
+		rmax_status_t Status = rmax_in_get_next_chunk(StreamId, MinChunkSize, MaxChunkSize, Timeout, Flags, &Completion);
+		if (Status == RMAX_OK)
+		{
+			ParseChunk(Completion);
+		}
+		else
+		{
+			UE_LOG(LogRivermax, Warning, TEXT("Rivermax Input stream failed to get next chunk. Status: %d"), Status);
+		}
+	}
+
+	bool FRivermaxInputStream::Init()
+	{
+		return true;
+	}
+
+	uint32 FRivermaxInputStream::Run()
+	{
+		while (bIsActive)
+		{
+			Process_AnyThread();
+		}
+
+		if (StreamId)
+		{
+			rmax_status_t Status = rmax_in_detach_flow(StreamId, &FlowAttribute);
+			if (Status != RMAX_OK)
+			{
+				UE_LOG(LogRivermax, Warning, TEXT("Failed to detach rivermax flow %d from input stream %d. Status: %d"), FlowAttribute.flow_id, StreamId, Status);
+			}
+
+			Status = rmax_in_destroy_stream(StreamId);
+
+			if (Status != RMAX_OK)
+			{
+				UE_LOG(LogRivermax, Warning, TEXT("Failed to destroy input stream %d correctly. Status: %d"), StreamId, Status);
+			}
+		}
+
+		return 0;
+	}
+
+	void FRivermaxInputStream::Stop()
+	{
+		bIsActive = false;
+	}
+
+	void FRivermaxInputStream::Exit()
+	{
+
+	}
+
+	bool FRivermaxInputStream::GetRTPParameter(uint8* InputRTP, FRTPParameter& OutParameter)
+	{
+		OutParameter.Timestamp = 0;
+		
+		const uint8* RTPHeaderPtr = GetRTPHeaderPointer(InputRTP);
+		if (((RTPHeaderPtr[0] & 0xC0) != 0x80))
+		{
+			return false;
+		}
+
+		OutParameter.SequencerNumber = RTPHeaderPtr[3] | RTPHeaderPtr[2] << 8;
+		if (IsExtendedSequenceNumber()) //todo if 2022 is supported extended sequence number is supported in other ways
+		{
+			OutParameter.SequencerNumber |= RTPHeaderPtr[12] << 24 | RTPHeaderPtr[13] << 16;
+			OutParameter.bIsFBit = !!(RTPHeaderPtr[16] & 0x80);
+		}
+		OutParameter.Timestamp = ByteSwap(*(uint32*)((RTPHeaderPtr)+4));
+		OutParameter.bIsMBit = !!(RTPHeaderPtr[1] & 0x80);
+		return true;
+	}
+
+	void FRivermaxInputStream::ParseChunk(const rmax_in_completion& Completion)
+	{
+		for (uint64 StrideIndex = 0; StrideIndex < Completion.chunk_size; ++StrideIndex)
+		{
+			ensure(Completion.hdr_ptr);
+			if (Completion.hdr_ptr == nullptr)
+			{
+				break;
+			}
+
+			uint8* HeaderPtr = (uint8*)Completion.hdr_ptr + StrideIndex * (size_t)BufferConfiguration.HeaderMemory.stride_size; //when using RMAX_RAW_PACKET then header is preceded by net header
+			uint8* DataPtr = (uint8*)Completion.data_ptr + StrideIndex * (size_t)BufferConfiguration.DataMemory.stride_size; // The payload is our data
+			uint32 DataOffset = 0;
+
+			if (Completion.packet_info_arr[StrideIndex].data_size)
+			{
+				FRTPParameter Parameter;
+				const bool bIsValid = GetRTPParameter(HeaderPtr, Parameter);
+				if (bIsValid)
+				{
+					if (bIsFirstFrameReceived)
+					{
+						//stats  rx count, received bytes++
+						uint64 LastSequenceNumberIncremented = StreamData.LastSequenceNumber + 1;
+						if (IsExtendedSequenceNumber() == false)
+						{
+							LastSequenceNumberIncremented = LastSequenceNumberIncremented & 0xFFFF;
+						}
+
+						bool bCanProcessChunk = true;
+						const uint64 LostPackets = ((uint64)Parameter.SequencerNumber + 0x100000000 - LastSequenceNumberIncremented) & 0xFFFFFFFF;
+						if (LostPackets > 0)
+						{
+							bCanProcessChunk = false;
+							StreamData.WritingOffset = 0;
+							StreamData.ReceivedSize = 0;
+							UE_LOG(LogRivermax, Warning, TEXT("Lost %uld packets"), LostPackets);
+						}
+
+						StreamData.LastSequenceNumber = Parameter.SequencerNumber;
+
+						//if flags are RMAX_IN_CREATE_STREAM_INFO_PER_PACKET todo 
+						{
+							if (FlowAttribute.flow_id && Completion.packet_info_arr[StrideIndex].flow_id != FlowAttribute.flow_id)
+							{
+								UE_LOG(LogRivermax, Error, TEXT("Received data from unexpected FlowId '%d'. Expected '%d'."), Completion.packet_info_arr[StrideIndex].flow_id, FlowAttribute.flow_id);
+							}
+						}
+
+						if (bCanProcessChunk)
+						{
+							TRACE_CPUPROFILER_EVENT_SCOPE(FRivermaxInputStream::ProcessingChunk)
+							const FRivermaxRTPSampleRowData* HeaderStart = reinterpret_cast<FRivermaxRTPSampleRowData*>(GetRTPHeaderPointer(HeaderPtr));
+
+							uint16 SizeToCopy = ByteSwap((uint16)HeaderStart->SRDLength1);
+							FMemory::Memcpy(&StreamData.CurrentFrame[StreamData.WritingOffset], &DataPtr[DataOffset], SizeToCopy);
+							StreamData.WritingOffset += SizeToCopy;
+							StreamData.ReceivedSize += SizeToCopy;
+
+							// todo Warning !!! GPUDirect doesn't support more than 1 SRD
+							if (HeaderStart->ContinuationBit1)
+							{
+								DataOffset += SizeToCopy;
+								SizeToCopy = ByteSwap((uint16)HeaderStart->SRDLength2);
+								FMemory::Memcpy(&StreamData.CurrentFrame[StreamData.WritingOffset], &DataPtr[DataOffset], SizeToCopy);
+								StreamData.WritingOffset += SizeToCopy;
+								StreamData.ReceivedSize += SizeToCopy;
+							}
+
+							if (StreamData.ReceivedSize > StreamData.ExpectedSize)
+							{
+								UE_LOG(LogRivermax, Warning, TEXT("Received too much data (%d). Expected %d but received (%d)"), StreamData.ReceivedSize - StreamData.ExpectedSize, StreamData.ExpectedSize, StreamData.ReceivedSize);
+								StreamData.WritingOffset = 0;
+								StreamData.ReceivedSize = 0;
+							}
+							else if (HeaderStart->MarkerBit)
+							{
+								if (StreamData.ReceivedSize == StreamData.ExpectedSize)
+								{
+									TRACE_CPUPROFILER_EVENT_SCOPE(FRivermaxInputStream::ProcessingReceivedFrame)
+									FRivermaxInputVideoFrameDescriptor Descriptor;
+									Descriptor.Width = Options.Resolution.X;
+									Descriptor.Height = Options.Resolution.Y;
+									Descriptor.Stride = Options.Resolution.X * 2;//todo stride. Bytes per row
+									FRivermaxInputVideoFrameReception NewFrame;
+									NewFrame.VideoBuffer = StreamData.CurrentFrame;
+									Listener->OnVideoFrameReceived(Descriptor, NewFrame);//todo stride
+									PrepareNextFrame();
+								}
+								else
+								{
+									UE_LOG(LogRivermax, Warning, TEXT("End of frame received (Marker bit) but not enough data was received (missing %d). Expected %d but received (%d)"), StreamData.ExpectedSize - StreamData.ReceivedSize, StreamData.ExpectedSize, StreamData.ReceivedSize);
+									StreamData.WritingOffset = 0;
+									StreamData.ReceivedSize = 0;
+								}
+							}
+						}
+
+					}
+					else
+					{
+						if (Parameter.bIsMBit)
+						{
+							StreamData.LastSequenceNumber = Parameter.SequencerNumber;
+							bIsFirstFrameReceived = true;
+							PrepareNextFrame();
+						}
+					}
+
+					//todo why  Reset RTP header
+					*(uint32_t*)GetRTPHeaderPointer(HeaderPtr) = 0;
+				}
+			}
+
+			//todo why?
+			Completion.packet_info_arr[StrideIndex].data_size = 0;
+		}
+	}
+
+	bool FRivermaxInputStream::IsExtendedSequenceNumber() const
+	{
+		return RivermaxStreamType == ERivermaxStreamType::VIDEO_2110_20_STREAM;
+	}
+
+	void FRivermaxInputStream::PrepareNextFrame()
+	{
+		using namespace UE::RivermaxCore::Private::Utils;
+
+		FRivermaxInputVideoFrameDescriptor Descriptor;
+		FRivermaxInputVideoFrameRequest Request;
+		const uint32 Groups = Options.Resolution.X / PixelsPerGroup_422_8b;
+		const uint32 BytesPerLine = Groups * BytesPerGroup_422_8b;
+		Descriptor.VideoBufferSize = Options.Resolution.Y * BytesPerLine;
+		Listener->OnVideoFrameRequested(Descriptor, Request);
+		StreamData.CurrentFrame = Request.VideoBuffer;
+		StreamData.WritingOffset = 0;
+		StreamData.ReceivedSize = 0;
+		StreamData.ExpectedSize = Descriptor.VideoBufferSize;
+	}
+}
+
+
+
+
+
diff --git a/Engine/Plugins/VirtualProduction/Rivermax/RivermaxCore/Source/Private/RivermaxInputStream.h b/Engine/Plugins/VirtualProduction/Rivermax/RivermaxCore/Source/Private/RivermaxInputStream.h
new file mode 100644
--- /dev/null
+++ b/Engine/Plugins/VirtualProduction/Rivermax/RivermaxCore/Source/Private/RivermaxInputStream.h
@@ -0,0 +1,94 @@
+// Copyright Epic Games, Inc. All Rights Reserved.
+
+#pragma once
+
+#include "IRivermaxInputStream.h"
+
+#include "HAL/Runnable.h"
+#include "RivermaxHeader.h"
+#include "RivermaxTypes.h"
+
+
+namespace UE::RivermaxCore::Private
+{
+	using UE::RivermaxCore::IRivermaxInputStream;
+	using UE::RivermaxCore::IRivermaxInputStreamListener;
+	using UE::RivermaxCore::FRivermaxStreamOptions;
+	
+	struct FInputStreamBufferConfiguration
+	{
+		size_t PayloadSize = 0;
+		size_t HeaderSize = 0;
+		uint16 PayloadExpectedSize = 1500;
+		uint16 HeaderExpectedSize = 20; //for 2110
+
+		rmax_in_memblock DataMemory;
+		rmax_in_memblock HeaderMemory;
+	};
+
+	struct FInputStreamData
+	{
+		uint64 LastSequenceNumber = 0;
+		uint8* CurrentFrame = nullptr;
+		uint32 WritingOffset = 0;
+		uint32 ReceivedSize = 0;
+		uint32 ExpectedSize = 0;
+	};
+
+	class FRivermaxInputStream : public IRivermaxInputStream, public FRunnable
+	{
+	public:
+		FRivermaxInputStream();
+		virtual ~FRivermaxInputStream();
+
+	public:
+
+		//~ Begin IRivermaxInputStreamListener interface
+		virtual bool Initialize(const FRivermaxStreamOptions& InOptions, IRivermaxInputStreamListener& InListener) override;
+		virtual void Uninitialize() override;
+		//~ End IRivermaxInputStreamListener interface 
+
+		void Process_AnyThread();
+
+		//~ Begin FRunnable interface
+		virtual bool Init() override;
+		virtual uint32 Run() override;
+		virtual void Stop() override;
+		virtual void Exit() override;
+		//~ End FRunnable interface
+
+	private:
+
+		struct FRTPParameter
+		{
+			uint32 SequencerNumber = 0;
+			uint32 Timestamp = 0;
+			bool bIsMBit = false;
+			bool bIsFBit = false;
+		};
+
+		bool GetRTPParameter(uint8* InputRTP, FRTPParameter& OutParameter);
+		void ParseChunk(const rmax_in_completion& Completion);
+		bool IsExtendedSequenceNumber() const;
+		void PrepareNextFrame();
+
+	private:
+
+
+
+		FRivermaxStreamOptions Options;
+
+		FRunnableThread* RivermaxThread = nullptr;
+		std::atomic<bool> bIsActive;
+
+		rmax_stream_id StreamId = 0;
+		rmax_in_flow_attr FlowAttribute;
+		FInputStreamBufferConfiguration BufferConfiguration;
+		bool bIsFirstFrameReceived = false;
+		ERivermaxStreamType RivermaxStreamType = ERivermaxStreamType::VIDEO_2110_20_STREAM;
+		FInputStreamData StreamData;
+		IRivermaxInputStreamListener* Listener;
+	};
+}
+
+
diff --git a/Engine/Plugins/VirtualProduction/Rivermax/RivermaxCore/Source/Private/RivermaxLog.cpp b/Engine/Plugins/VirtualProduction/Rivermax/RivermaxCore/Source/Private/RivermaxLog.cpp
new file mode 100644
--- /dev/null
+++ b/Engine/Plugins/VirtualProduction/Rivermax/RivermaxCore/Source/Private/RivermaxLog.cpp
@@ -0,0 +1,5 @@
+// Copyright Epic Games, Inc. All Rights Reserved.
+
+#include "RivermaxLog.h"
+
+DEFINE_LOG_CATEGORY(LogRivermax);
diff --git a/Engine/Plugins/VirtualProduction/Rivermax/RivermaxCore/Source/Private/RivermaxLog.h b/Engine/Plugins/VirtualProduction/Rivermax/RivermaxCore/Source/Private/RivermaxLog.h
new file mode 100644
--- /dev/null
+++ b/Engine/Plugins/VirtualProduction/Rivermax/RivermaxCore/Source/Private/RivermaxLog.h
@@ -0,0 +1,7 @@
+// Copyright Epic Games, Inc. All Rights Reserved.
+
+#pragma once
+
+#include "CoreMinimal.h"
+
+DECLARE_LOG_CATEGORY_EXTERN(LogRivermax, Log, All);
diff --git a/Engine/Plugins/VirtualProduction/Rivermax/RivermaxCore/Source/Private/RivermaxManager.cpp b/Engine/Plugins/VirtualProduction/Rivermax/RivermaxCore/Source/Private/RivermaxManager.cpp
new file mode 100644
--- /dev/null
+++ b/Engine/Plugins/VirtualProduction/Rivermax/RivermaxCore/Source/Private/RivermaxManager.cpp
@@ -0,0 +1,112 @@
+// Copyright Epic Games, Inc. All Rights Reserved.
+
+#include "RivermaxManager.h"
+
+#include "Misc/App.h"
+#include "Misc/Paths.h"
+#include "Misc/CoreDelegates.h"
+#include "RivermaxHeader.h"
+#include "RivermaxLog.h"
+
+
+namespace UE::RivermaxCore::Private
+{
+	FRivermaxManager::FRivermaxManager()
+	{
+		const auto InitializeRivermaxLibFunc = [this]()
+		{
+			rmax_init_config Config;
+			memset(&Config, 0, sizeof(Config));
+			//Config.flags |= RIVERMAX_HANDLE_SIGNAL; //Verify if desired
+
+			rmax_status_t Status = rmax_init(&Config);
+			if (Status == RMAX_OK)
+			{
+				uint32 Major = 0;
+				uint32 Minor = 0;
+				uint32 Release = 0;
+				uint32 Build = 0;
+				Status = rmax_get_version(&Major, &Minor, &Release, &Build);
+				if (Status == RMAX_OK)
+				{
+					bIsInitialized = true;
+					UE_LOG(LogRivermax, Log, TEXT("Rivermax library version %d.%d.%d.%d succesfully initialized"), Major, Minor, Release, Build);
+				}
+				else
+				{
+					UE_LOG(LogRivermax, Log, TEXT("Failed to retrieve Rivermax library version. Status: %d"), Status);
+				}
+			}
+			else if (Status == RMAX_ERR_LICENSE_ISSUE)
+			{
+				UE_LOG(LogRivermax, Log, TEXT("Rivermax License could not be found. Have you configured RIVERMAX_LICENSE_PATH environment variable?"));
+			}
+			
+			if(bIsInitialized == false)
+			{
+				UE_LOG(LogRivermax, Log, TEXT("Rivermax library failed to initialize. Status: %d"), Status);
+			}
+		};
+
+		const bool bIsLibraryLoaded = LoadRivermaxLibrary();
+		if (bIsLibraryLoaded && FApp::CanEverRender())
+		{
+			//Postpone initialization after all modules have been loaded
+			FCoreDelegates::OnAllModuleLoadingPhasesComplete.AddLambda(InitializeRivermaxLibFunc);
+		}
+		else
+		{
+			UE_LOG(LogRivermax, Log, TEXT("Skipping Rivermax initialization. Library won't be usable."));
+		}
+	}
+
+	FRivermaxManager::~FRivermaxManager()
+	{
+		if (IsInitialized())
+		{
+			rmax_status_t Status = rmax_cleanup();
+			if (Status != RMAX_OK)
+			{
+				UE_LOG(LogRivermax, Log, TEXT("Failed to cleanup Rivermax Library. Status: %d"), Status);
+			}
+			UE_LOG(LogRivermax, Log, TEXT("Rivermax Library has shutdown."));
+		}
+
+		if (LibraryHandle != nullptr)
+		{
+			FPlatformProcess::FreeDllHandle(LibraryHandle);
+			LibraryHandle = nullptr;
+		}
+	}
+
+	bool FRivermaxManager::IsInitialized() const
+	{
+		return bIsInitialized;
+	}
+
+	bool FRivermaxManager::LoadRivermaxLibrary()
+	{
+#if defined(RIVERMAX_LIBRARY_PLATFORM_PATH) && defined(RIVERMAX_LIBRARY_NAME)
+		const FString LibraryPath = FPaths::Combine(FPaths::EngineSourceDir(), TEXT("ThirdParty/NVIDIA/Rivermax/lib"), TEXT(PREPROCESSOR_TO_STRING(RIVERMAX_LIBRARY_PLATFORM_PATH)));
+		const FString LibraryName = TEXT(PREPROCESSOR_TO_STRING(RIVERMAX_LIBRARY_NAME));
+
+		FPlatformProcess::PushDllDirectory(*LibraryPath);
+		LibraryHandle = FPlatformProcess::GetDllHandle(*LibraryName);
+		FPlatformProcess::PopDllDirectory(*LibraryPath);
+#endif
+
+		const bool bIsLibraryValid = LibraryHandle != nullptr;
+		if(!bIsLibraryValid)
+		{
+			UE_LOG(LogRivermax, Log, TEXT("Failed to load required library %s. Rivermax library will not be functional."), *LibraryName);
+		}
+
+		return bIsLibraryValid;
+	}
+
+}
+
+
+
+
+
diff --git a/Engine/Plugins/VirtualProduction/Rivermax/RivermaxCore/Source/Private/RivermaxManager.h b/Engine/Plugins/VirtualProduction/Rivermax/RivermaxCore/Source/Private/RivermaxManager.h
new file mode 100644
--- /dev/null
+++ b/Engine/Plugins/VirtualProduction/Rivermax/RivermaxCore/Source/Private/RivermaxManager.h
@@ -0,0 +1,30 @@
+// Copyright Epic Games, Inc. All Rights Reserved.
+
+#pragma once
+
+#include "IRivermaxManager.h"
+
+
+namespace UE::RivermaxCore::Private
+{
+	class FRivermaxManager : public UE::RivermaxCore::IRivermaxManager
+	{
+	public:
+		FRivermaxManager();
+		~FRivermaxManager();
+
+	public:
+		//~ Begin IRivermaxManager interface
+		virtual bool IsInitialized() const override;
+		//~ End IRivermaxManager interface
+
+	private:
+		bool LoadRivermaxLibrary();
+
+	private:
+		bool bIsInitialized = false;
+		void* LibraryHandle = nullptr;
+	};
+}
+
+
diff --git a/Engine/Plugins/VirtualProduction/Rivermax/RivermaxCore/Source/Private/RivermaxOutputFrame.cpp b/Engine/Plugins/VirtualProduction/Rivermax/RivermaxCore/Source/Private/RivermaxOutputFrame.cpp
new file mode 100644
--- /dev/null
+++ b/Engine/Plugins/VirtualProduction/Rivermax/RivermaxCore/Source/Private/RivermaxOutputFrame.cpp
@@ -0,0 +1,36 @@
+// Copyright Epic Games, Inc. All Rights Reserved.
+
+#include "RivermaxOutputFrame.h"
+
+namespace UE::RivermaxCore::Private
+{
+	FRivermaxOutputFrame::FRivermaxOutputFrame(uint32 InFrameIndex)
+		: FrameIndex(InFrameIndex)
+		, VideoBuffer(nullptr)
+	{
+		Clear();
+	}
+
+	FRivermaxOutputFrame::~FRivermaxOutputFrame()
+	{
+		if (VideoBuffer)
+		{
+			FMemory::Free(VideoBuffer);
+		}
+	}
+
+	bool FRivermaxOutputFrame::IsReadyToBeSent() const
+	{
+		return FrameIdentifier != InvalidIdentifier && bIsVideoBufferReady;
+	}
+
+	void FRivermaxOutputFrame::Clear()
+	{
+		FrameIdentifier = InvalidIdentifier;
+		PacketCounter = 0;
+		LineNumber = 0;
+		SRDOffset = 0;
+		TimestampTicks = 0.0;
+		bIsVideoBufferReady = false;
+	}
+}
diff --git a/Engine/Plugins/VirtualProduction/Rivermax/RivermaxCore/Source/Private/RivermaxOutputFrame.h b/Engine/Plugins/VirtualProduction/Rivermax/RivermaxCore/Source/Private/RivermaxOutputFrame.h
new file mode 100644
--- /dev/null
+++ b/Engine/Plugins/VirtualProduction/Rivermax/RivermaxCore/Source/Private/RivermaxOutputFrame.h
@@ -0,0 +1,61 @@
+// Copyright Epic Games, Inc. All Rights Reserved.
+
+#pragma once
+
+#include "CoreMinimal.h"
+
+#include "Math/NumericLimits.h"
+#include "Misc/Timecode.h"
+
+namespace UE::RivermaxCore::Private
+{
+	/**
+	 * Output frame descriptor. Contains data to be sent on the wire and packetization tracking
+	 */
+	struct FRivermaxOutputFrame
+	{
+		FRivermaxOutputFrame(uint32 InFrameIndex);
+		~FRivermaxOutputFrame();
+		FRivermaxOutputFrame(const FRivermaxOutputFrame&) = delete;
+		FRivermaxOutputFrame& operator=(const FRivermaxOutputFrame&) = delete;
+
+		/** Whether this frame is ready to be sent */
+		bool IsReadyToBeSent() const;
+
+		/** Reset internal counters of this frame */
+		void Clear();
+
+		/** Index of the frame used as unique identifier */
+		uint32 FrameIndex = 0;
+
+		/** Video buffer where we copy texture to be sent out */
+		uint8* VideoBuffer = nullptr;
+
+		/** Identifier of the frame incremented every time a frame is captured */
+		static constexpr uint32 InvalidIdentifier = MAX_uint32;
+		uint32 FrameIdentifier = InvalidIdentifier;
+
+		/** Timecode at which this frame was captured */
+		FTimecode Timecode;
+
+		/** Used to flag this frame ready to be sent */
+		bool bIsVideoBufferReady = false;
+
+		/** Variables about progress of packetization of the frame */
+		uint32 PacketCounter = 0;
+		uint32 LineNumber = 0;
+		uint16 SRDOffset = 0;
+		uint32 ChunkNumber = 0;
+
+		/** Timestamp of this frame  */
+		double TimestampTicks = 0.0;
+
+		/** Pointers retrieved from Rivermax for the next chunk */
+		void* PayloadPtr = nullptr;
+		void* HeaderPtr = nullptr;
+
+		
+	};
+}
+
+
diff --git a/Engine/Plugins/VirtualProduction/Rivermax/RivermaxCore/Source/Private/RivermaxOutputStream.cpp b/Engine/Plugins/VirtualProduction/Rivermax/RivermaxCore/Source/Private/RivermaxOutputStream.cpp
new file mode 100644
--- /dev/null
+++ b/Engine/Plugins/VirtualProduction/Rivermax/RivermaxCore/Source/Private/RivermaxOutputStream.cpp
@@ -0,0 +1,672 @@
+// Copyright Epic Games, Inc. All Rights Reserved.
+
+#include "RivermaxOutputStream.h"
+
+#include "Async/Async.h"
+#include "IRivermaxCoreModule.h"
+#include "IRivermaxManager.h"
+#include "Misc/ByteSwap.h"
+#include "RivermaxLog.h"
+#include "RivermaxTypes.h"
+#include "RivermaxUtils.h"
+
+
+namespace UE::RivermaxCore::Private
+{
+	void CalculateStreamTime(const FRivermaxOutputStreamMemory& InStreamMemory, const FRivermaxOutputStreamData& InStreamData, const FRivermaxStreamOptions& InMediaOptions, const double InSampleRate, double& InOutTimeNs, double& OutTimestampTick)
+	{
+		using namespace UE::RivermaxCore::Private::Utils;
+
+		double FrameIntervalNs = InStreamData.FrameFieldTimeIntervalNs;
+		const bool bIsProgressive = true;//todo MediaConfiguration.IsProgressive() 
+		const ERivermaxStreamType StreamType = ERivermaxStreamType::VIDEO_2110_20_STREAM; //todo
+		if (bIsProgressive == false)
+		{
+			FrameIntervalNs *= 2;
+		}
+		const uint64 FrameAtTime = (uint64)(InOutTimeNs / FrameIntervalNs + 1);
+		double FirstPacketStartTimeNs = FrameAtTime * FrameIntervalNs; //next alignment point calculation
+
+		uint32 PacketsInFrameField = InStreamMemory.PacketsInFrameField;
+		if (bIsProgressive == false)
+		{
+			PacketsInFrameField *= 2;
+		}
+		
+		// See https://ieeexplore.ieee.org/stamp/stamp.jsp?tp=&arnumber=8165971 for reference
+		// To be revisited for non standard resolution/frame rate
+		if (StreamType == ERivermaxStreamType::VIDEO_2110_20_STREAM)
+		{
+			double RActive;
+			double TRODefaultMultiplier;
+			if (bIsProgressive)
+			{
+				RActive = (1080.0 / 1125.0);
+				if (InMediaOptions.Resolution.Y >= FullHDHeight)
+				{
+				 	// As defined by SMPTE 2110-21 6.3.2
+					TRODefaultMultiplier = (43.0 / 1125.0);
+				}
+				else 
+				{
+					TRODefaultMultiplier = (28.0 / 750.0);
+				}
+			}
+			else
+			{
+				if (InMediaOptions.Resolution.Y >= FullHDHeight)
+				{
+					// As defined by SMPTE 2110-21 6.3.3
+					RActive = (1080.0 / 1125.0);
+					TRODefaultMultiplier = (22.0 / 1125.0);
+				}
+				else if (InMediaOptions.Resolution.Y >= 576)
+				{
+					RActive = (576.0 / 625.0);
+					TRODefaultMultiplier = (26.0 / 625.0);
+				}
+				else 
+				{
+					RActive = (487.0 / 525.0);
+					TRODefaultMultiplier = (20.0 / 525.0);
+				}
+			}
+
+			const double TRSNano = (FrameIntervalNs * RActive) / PacketsInFrameField;
+			const double TROffset = (TRODefaultMultiplier * FrameIntervalNs) - (VideoTROModification * TRSNano);
+			FirstPacketStartTimeNs += TROffset;
+		}
+
+		constexpr double OneSecondNs = 1 * 1E9;
+		OutTimestampTick = FirstPacketStartTimeNs * (InSampleRate / OneSecondNs);
+		InOutTimeNs = FirstPacketStartTimeNs;
+	}
+	
+	uint32 FindPayloadSize(const FRivermaxStreamOptions& InOptions, uint32 InBytesPerLine)
+	{
+		using namespace UE::RivermaxCore::Private::Utils;
+
+		// If stride per line is small then our minimal payload size, send one line per packet.
+		if (InBytesPerLine <= MaxPayloadSize)
+		{
+			return InBytesPerLine;
+		}
+
+		uint32 PacketDivider = 2; //Start with 2 packets per line as 1 is too big already
+		static constexpr uint32 MaxPacketCountPerLine = 10; 
+		while (PacketDivider <= MaxPacketCountPerLine)
+		{
+			const uint32 TestPayloadSize = InBytesPerLine / PacketDivider;
+			const bool bIsFlushInLine = InBytesPerLine % PacketDivider == 0;
+			if (bIsFlushInLine && TestPayloadSize < MaxPayloadSize)
+			{
+				return TestPayloadSize;
+			}
+
+			++PacketDivider;
+		}
+
+		return 0;
+	}
+
+	uint32 FindChunksPerLine(const FRivermaxStreamOptions& InOptions)
+	{
+		uint32 ChunksPerLine = 3; //Start with 3 chunks with line. Todo : Ask Nvidia
+		static constexpr uint32 MaxChunksPerLine = 10;
+		while (ChunksPerLine <= MaxChunksPerLine)
+		{
+			if (InOptions.Resolution.Y % ChunksPerLine == 0)
+			{
+				return ChunksPerLine;
+			}
+
+			++ChunksPerLine;
+		}
+
+		return 0;
+	}
+
+
+	FRivermaxOutputStream::FRivermaxOutputStream()
+		: bIsActive(false)
+	{
+
+	}
+
+	FRivermaxOutputStream::~FRivermaxOutputStream()
+	{
+		Uninitialize();
+	}
+
+	bool FRivermaxOutputStream::Initialize(const FRivermaxStreamOptions& InOptions, IRivermaxOutputStreamListener& InListener)
+	{
+		IRivermaxCoreModule& RivermaxModule = FModuleManager::LoadModuleChecked<IRivermaxCoreModule>(TEXT("RivermaxCore"));
+		if (RivermaxModule.GetRivermaxManager()->IsInitialized() == false)
+		{
+			UE_LOG(LogRivermax, Warning, TEXT("Can't create Rivermax Output Stream. Library isn't initialized."));
+			return false;
+		}
+
+		Options = InOptions;
+		Listener = &InListener;
+
+		Async(EAsyncExecution::TaskGraph, [this]()
+		{
+			//Create event to wait on when no frames are available to send
+			ReadyToSendEvent = FPlatformProcess::GetSynchEventFromPool();
+
+			TAnsiStringBuilder<2048> SDPDescription;
+			UE::RivermaxCore::Private::Utils::StreamOptionsToSDPDescription(Options, SDPDescription);
+
+			// Initialize video buffer memory we will be using
+			InitializeBuffers();
+
+			// Initialize internal memory and rivermax configuration 
+			InitializeMemory();
+
+			// Initialize rivermax output stream with desired config
+			const uint32 NumberPacketsPerFrame = StreamMemory.PacketsInFrameField;
+			const uint32 MediaBlockIndex = 0;
+			rmax_stream_id NewId;
+			rmax_qos_attr QOSAttributes = { 0, 0 }; //todo
+
+			rmax_buffer_attr BufferAttributes;
+			BufferAttributes.chunk_size_in_strides = StreamMemory.ChunkSizeInStrides;
+			BufferAttributes.data_stride_size = StreamMemory.DataStrideSize;
+			BufferAttributes.app_hdr_stride_size = StreamMemory.HeaderStrideSize; //Todo RTP header size
+			BufferAttributes.mem_block_array = StreamMemory.MemoryBlocks.GetData();
+			BufferAttributes.mem_block_array_len = StreamMemory.MemoryBlocks.Num();
+
+			rmax_status_t Status = rmax_out_create_stream(SDPDescription.GetData(), &BufferAttributes, &QOSAttributes, NumberPacketsPerFrame, MediaBlockIndex, &NewId);
+			if (Status == RMAX_OK)
+			{
+				struct sockaddr_in SourceAddress;
+				struct sockaddr_in DestinationAddress;
+				Status = rmax_out_query_address(NewId, MediaBlockIndex, &SourceAddress, &DestinationAddress);
+				if (Status == RMAX_OK)
+				{
+					StreamId = NewId;
+
+					// Todo add event manager to avoid spinning
+
+					StreamData.FrameFieldTimeIntervalNs = 1E9 / Options.FrameRate.AsDecimal();
+					StreamData.SendTimeNs = (FPlatformTime::Seconds() + 1) * 1E9; //Current time + 1 sec todo why?
+					CalculateStreamTime(StreamMemory, StreamData, Options, MediaClockSampleRate, StreamData.SendTimeNs, StreamData.InitialTimestampTick);
+					StreamData.StartSendTimeNs = StreamData.SendTimeNs;
+
+					bIsActive = true;
+					RivermaxThread.Reset(FRunnableThread::Create(this, TEXT("Rmax OutputStream Thread"), 128 * 1024, TPri_AboveNormal, FPlatformAffinity::GetPoolThreadMask()));
+				}
+			}
+			else
+			{
+				UE_LOG(LogRivermax, Warning, TEXT("Failed to create Rivermax output stream. Status: %d"), Status);
+			}
+
+			Listener->OnInitializationCompleted(bIsActive);
+		});
+
+		return true;
+	}
+
+	void FRivermaxOutputStream::Uninitialize()
+	{
+		if (RivermaxThread != nullptr)
+		{
+			Stop();
+			
+			ReadyToSendEvent->Trigger();
+			RivermaxThread->Kill(true);
+			
+			RivermaxThread.Reset();
+
+
+			FPlatformProcess::ReturnSynchEventToPool(ReadyToSendEvent);
+			ReadyToSendEvent = nullptr;
+			UE_LOG(LogRivermax, Log, TEXT("Rivermax Output stream has shutdown"));
+		}
+
+	}
+
+	bool FRivermaxOutputStream::PushVideoFrame(const FRivermaxOutputVideoFrameInfo& NewFrame)
+	{
+		if (TSharedPtr<FRivermaxOutputFrame> AvailableFrame = GetNextAvailableFrame(NewFrame.FrameIdentifier))
+		{
+			FMemory::Memcpy(AvailableFrame->VideoBuffer, NewFrame.VideoBuffer, NewFrame.Height * NewFrame.Stride);
+			AvailableFrame->bIsVideoBufferReady = true;
+
+			//If Frame ready to be sent
+			if(AvailableFrame->IsReadyToBeSent())
+			{
+				const FString TraceName = FString::Format(TEXT("FRivermaxOutputStream::PushFrame {0}"), { AvailableFrame->FrameIndex });
+				TRACE_CPUPROFILER_EVENT_SCOPE_TEXT(*TraceName);
+
+				FScopeLock Lock(&FrameCriticalSection);
+				
+				AvailableFrames.Remove(AvailableFrame);
+				FramesToSend.Add(MoveTemp(AvailableFrame));
+				ReadyToSendEvent->Trigger();
+			}
+
+			return true;
+		}
+		
+		return false;
+	}
+
+	void FRivermaxOutputStream::Process_AnyThread()
+	{
+		if (CurrentFrame.IsValid())
+		{
+			{
+				TRACE_CPUPROFILER_EVENT_SCOPE(FRivermaxOutputStream::Wait);
+				StreamData.SendTimeNs = StreamData.StartSendTimeNs + StreamData.FrameFieldTimeIntervalNs * StreamMemory.FramesFieldPerMemoryBlock * Stats.MemoryBlockSentCounter;
+				WaitForNextRound(StreamData.SendTimeNs);
+			}
+
+			const FString TraceName = FString::Format(TEXT("FRivermaxOutputStream::SendFrame {0}"), { CurrentFrame->FrameIndex });
+			TRACE_CPUPROFILER_EVENT_SCOPE_TEXT(*TraceName);
+
+			do 
+			{
+				if (bIsActive)
+				{
+					GetNextChunk();
+				}
+
+				if (bIsActive)
+				{
+					SetupRTPHeaders();
+				}
+
+				if (bIsActive)
+				{
+					CommitNextChunks();
+				}
+
+				//Update frame progress
+				if (bIsActive)
+				{
+					Stats.TotalStrides += StreamMemory.ChunkSizeInStrides;
+
+					if ((++CurrentFrame->ChunkNumber) % StreamMemory.ChunksPerFrameField == 0)
+					{
+						StreamData.SendTimeNs += StreamData.FrameFieldTimeIntervalNs;
+					}
+				}
+
+			} while (CurrentFrame->ChunkNumber < StreamMemory.ChunksPerMemoryBlock && bIsActive);
+			
+			//We are done with this frame, put it back as available
+			CurrentFrame->Clear();
+			{
+				FScopeLock Lock(&FrameCriticalSection);
+				AvailableFrames.Add(MoveTemp(CurrentFrame));
+				CurrentFrame.Reset();
+			}
+
+			Stats.MemoryBlockSentCounter++;
+			Stats.TotalStrides += StreamMemory.ChunkSizeInStrides;
+			StreamData.bHasFrameFirstChunkBeenFetched = false;
+		}
+	}
+
+	void FRivermaxOutputStream::InitializeBuffers()
+	{
+		AvailableFrames.Reserve(Options.NumberOfBuffers);
+		for (int32 Index = 0; Index < Options.NumberOfBuffers; ++Index)
+		{
+			TSharedPtr<FRivermaxOutputFrame> Frame = MakeShared<FRivermaxOutputFrame>(Index);
+			Frame->VideoBuffer = static_cast<uint8*>(FMemory::Malloc(Options.Resolution.Y * Options.Resolution.X * 4)); //todo alignment + bytes per row
+			AvailableFrames.Add(MoveTemp(Frame));
+		}
+	}
+
+	void FRivermaxOutputStream::InitializeMemory()
+	{
+		using namespace UE::RivermaxCore::Private::Utils;
+
+		//2110 stream type
+		const uint32 Groups = FMath::CeilToInt32(Options.Resolution.X / (float)PixelsPerGroup_422_8b); //This will need to be updated based on pixel format
+		const uint32 BytesPerLine = Groups * BytesPerGroup_422_8b;
+		StreamMemory.PayloadSize = FindPayloadSize(Options, BytesPerLine); 
+		StreamMemory.DataStrideSize = StreamMemory.PayloadSize;//todo align cache line?
+		StreamMemory.LinesInChunk = FindChunksPerLine(Options);
+
+
+		StreamMemory.PacketsInLine =  BytesPerLine / (float)StreamMemory.PayloadSize;
+		StreamMemory.ChunkSizeInStrides = StreamMemory.LinesInChunk * StreamMemory.PacketsInLine;
+
+		StreamMemory.FramesFieldPerMemoryBlock = 1;
+		StreamMemory.PacketsInFrameField =  StreamMemory.PacketsInLine * Options.Resolution.Y;
+		StreamMemory.PacketsPerMemoryBlock = StreamMemory.PacketsInFrameField * StreamMemory.FramesFieldPerMemoryBlock;
+		StreamMemory.ChunksPerFrameField = StreamMemory.PacketsInFrameField / StreamMemory.ChunkSizeInStrides;
+		StreamMemory.ChunksPerMemoryBlock = StreamMemory.FramesFieldPerMemoryBlock * StreamMemory.ChunksPerFrameField;
+		StreamMemory.MemoryBlockCount = Options.NumberOfBuffers; //todo
+		StreamMemory.StridesPerMemoryBlock = StreamMemory.ChunkSizeInStrides * StreamMemory.ChunksPerMemoryBlock;
+
+		// Setup arrays with the right sizes so we can give pointers to rivermax
+		StreamMemory.RTPHeaders.SetNumZeroed(StreamMemory.PacketsPerMemoryBlock);
+		StreamMemory.PayloadSizes.SetNumUninitialized(StreamMemory.PacketsPerMemoryBlock);
+		StreamMemory.HeaderSizes.SetNumUninitialized(StreamMemory.PacketsPerMemoryBlock);
+		StreamMemory.HeaderStrideSize = BytesPerHeader;
+		for (int32 PayloadSizeIndex = 0; PayloadSizeIndex < StreamMemory.PayloadSizes.Num(); ++PayloadSizeIndex)
+		{
+			StreamMemory.PayloadSizes[PayloadSizeIndex] = StreamMemory.PayloadSize;
+			StreamMemory.HeaderSizes[PayloadSizeIndex] = StreamMemory.HeaderStrideSize;
+		}
+		StreamMemory.MemoryBlocks.SetNum(StreamMemory.MemoryBlockCount);
+		for (uint32 BlockIndex = 0; BlockIndex < StreamMemory.MemoryBlockCount; ++BlockIndex)
+		{
+			rmax_mem_block& Block = StreamMemory.MemoryBlocks[BlockIndex];
+			Block.chunks_num = StreamMemory.ChunksPerMemoryBlock;
+			Block.app_hdr_size_arr = StreamMemory.HeaderSizes.GetData();
+			Block.data_size_arr = StreamMemory.PayloadSizes.GetData();
+			Block.data_ptr = AvailableFrames[BlockIndex]->VideoBuffer; //This doesn't work with more than 2 samples. Need to rethink how we preassign frame address in rivermax
+			Block.app_hdr_ptr = &StreamMemory.RTPHeaders[BlockIndex];
+		}
+	}
+
+	void FRivermaxOutputStream::InitializeNextFrame(const TSharedPtr<FRivermaxOutputFrame>& NextFrame)
+	{
+		NextFrame->TimestampTicks = StreamData.InitialTimestampTick;
+		NextFrame->LineNumber = 0;
+		NextFrame->PacketCounter = 0;
+		NextFrame->SRDOffset = 0;
+		NextFrame->ChunkNumber = 0;
+		NextFrame->PayloadPtr = nullptr;
+		NextFrame->HeaderPtr = nullptr;
+	}
+
+	TSharedPtr<FRivermaxOutputFrame> FRivermaxOutputStream::GetNextFrameToSend()
+	{
+		TSharedPtr<FRivermaxOutputFrame> FrameToSend;
+
+		FScopeLock Lock(&FrameCriticalSection);
+		if (FramesToSend.IsEmpty() == false)
+		{
+			// Pick oldest frame of the lot. 
+			// Todo Sort array when modified to always pick first item.
+			uint32 OldestIdentifier = TNumericLimits<uint32>::Max();
+			TArray<TSharedPtr<FRivermaxOutputFrame>>::TIterator It = FramesToSend.CreateIterator();
+			for (It; It; ++It)
+			{
+				TSharedPtr<FRivermaxOutputFrame>& ItFrame = *It;
+				if (ItFrame->FrameIdentifier < OldestIdentifier)
+				{
+					FrameToSend = *It;
+					OldestIdentifier = FrameToSend->FrameIdentifier;
+				}
+			}
+
+			if (ensure(FrameToSend))
+			{
+				FramesToSend.Remove(FrameToSend);
+				InitializeNextFrame(FrameToSend);
+			}
+		}
+
+		return FrameToSend;
+	}
+
+	TSharedPtr<FRivermaxOutputFrame> FRivermaxOutputStream::GetNextAvailableFrame(uint32 InFrameIdentifier)
+	{
+		TSharedPtr<FRivermaxOutputFrame> NextFrame;
+
+		{
+			FScopeLock Lock(&FrameCriticalSection);
+
+			//Find matching frame identifier
+			if(TSharedPtr<FRivermaxOutputFrame>* MatchingFrame = AvailableFrames.FindByPredicate([InFrameIdentifier](const TSharedPtr<FRivermaxOutputFrame>& Other){ return Other->FrameIdentifier == InFrameIdentifier; }))
+			{
+				NextFrame = *MatchingFrame;
+			}
+			else if(TSharedPtr<FRivermaxOutputFrame>* EmptyFrame = AvailableFrames.FindByPredicate([](const TSharedPtr<FRivermaxOutputFrame>& Other){ return Other->FrameIdentifier == FRivermaxOutputFrame::InvalidIdentifier; })) //Find an empty frame
+			{
+				NextFrame = *EmptyFrame;
+				NextFrame->FrameIdentifier = InFrameIdentifier;
+			}
+			else
+			{
+				// Could reuse a frame not yet ready to be sent
+				// Could also look into stealing a frame ready to be sent bot not sent yet
+			}
+		}
+
+		return NextFrame;
+	}
+
+	void FRivermaxOutputStream::BuildRTPHeader(FRTPHeader& OutHeader) const
+	{
+		// build RTP header - 12 bytes
+		/*
+		0                   1                   2                   3
+		0 1 2 3 4 5 6 7 8 9 0 1 2 3 4 5 6 7 8 9 0 1 2 3 4 5 6 7 8 9 0 1
+		+-+-+-+-+-+-+-+-+-+-+-+-+-+-+-+-+-+-+-+-+-+-+-+-+-+-+-+-+-+-+-+-+
+		| V |P|X|  CC   |M|     PT      |            SEQ                |
+		+-+-+-+-+-+-+-+-+-+-+-+-+-+-+-+-+-+-+-+-+-+-+-+-+-+-+-+-+-+-+-+-+
+		|                           timestamp                           |
+		+-+-+-+-+-+-+-+-+-+-+-+-+-+-+-+-+-+-+-+-+-+-+-+-+-+-+-+-+-+-+-+-+
+		|                           ssrc                                |
+		+-+-+-+-+-+-+-+-+-+-+-+-+-+-+-+-+-+-+-+-+-+-+-+-+-+-+-+-+-+-+-+-+*/
+		OutHeader.RawHeader[0] = 0x80;  // 10000000 - version2, no padding, no extension
+		OutHeader.RawHeader[1] = 96; //todo payload type from sdp file
+		OutHeader.RawHeader[2] = (StreamData.SequenceNumber >> 8) & 0xFF;  // sequence number
+		OutHeader.RawHeader[3] = (StreamData.SequenceNumber) & 0xFF;  // sequence number
+		*(uint32*)&OutHeader.RawHeader[4] = ByteSwap((uint32)CurrentFrame->TimestampTicks);
+
+		//2110 specific header
+		*(uint32*)&OutHeader.RawHeader[8] = 0x0eb51dbd;  // simulated ssrc
+
+		if (StreamType == ERivermaxStreamType::VIDEO_2110_20_STREAM)
+		{
+			//SRD means Sample Row Data
+			// build SRD header - 8-14 bytes
+		   /* 0                   1                   2                   3
+			0 1 2 3 4 5 6 7 8 9 0 1 2 3 4 5 6 7 8 9 0 1 2 3 4 5 6 7 8 9 0 1
+			+-+-+-+-+-+-+-+-+-+-+-+-+-+-+-+-+-+-+-+-+-+-+-+-+-+-+-+-+-+-+-+-+
+			|    Extended Sequence Number   |           SRD Length          |
+			+-+-+-+-+-+-+-+-+-+-+-+-+-+-+-+-+-+-+-+-+-+-+-+-+-+-+-+-+-+-+-+-+
+			|F|     SRD Row Number          |C|         SRD Offset          |
+			+-+-+-+-+-+-+-+-+-+-+-+-+-+-+-+-+-+-+-+-+-+-+-+-+-+-+-+-+-+-+-+-+ */
+			OutHeader.RawHeader[12] = (StreamData.SequenceNumber >> 24) & 0xff;  // high 16 bit of seq Extended Sequence Number
+			OutHeader.RawHeader[13] = (StreamData.SequenceNumber >> 16) & 0xff;  // high 16 bit of seq Extended Sequence Number
+
+			*(uint16*)&OutHeader.RawHeader[14] = ByteSwap(StreamMemory.PayloadSize /*- 20*/);  // SRD Length
+		
+			uint16 number_of_rows = Options.Resolution.Y; //todo divide by 2 if interlaced
+			uint16 srd_row_number = (CurrentFrame->LineNumber % number_of_rows);
+			*(uint16*)&OutHeader.RawHeader[16] = ByteSwap(srd_row_number);
+			OutHeader.RawHeader[16] |= (0 << 7); //todo when fields are sent for interlace
+
+			// we never have continuation
+			*(uint16*)&OutHeader.RawHeader[18] = ByteSwap(CurrentFrame->SRDOffset);  // SRD Offset
+			uint16 group_size = (uint16)((StreamMemory.PayloadSize /*- 20*/) / 2.5);
+			CurrentFrame->SRDOffset = (CurrentFrame->SRDOffset + group_size) % (group_size * StreamMemory.PacketsInLine);
+		}
+
+		if (++CurrentFrame->PacketCounter == StreamMemory.PacketsInFrameField)
+		{
+			OutHeader.RawHeader[1] |= 0x80; // last packet in frame (Marker)
+			// ST2210-20: the timestamp SHOULD be the same for each packet of the frame/field.
+			const double ticks = (MediaClockSampleRate / (Options.FrameRate.AsDecimal()));
+			//if (set.video_type != VIDEO_TYPE::PROGRESSIVE) {
+			//	send_data.m_second_field = !send_data.m_second_field;
+			//	ticks /= 2;
+			//}
+			CurrentFrame->TimestampTicks += ticks;
+		}
+	}
+
+	void FRivermaxOutputStream::DestroyStream()
+	{
+		rmax_status_t Status = rmax_out_cancel_unsent_chunks(StreamId);
+		if (Status != RMAX_OK)
+		{
+			UE_LOG(LogRivermax, Warning, TEXT("Could not cancel unsent chunks when destroying output stream. Status: %d"), Status);
+		}
+
+		do 
+		{
+			Status = rmax_out_destroy_stream(StreamId);
+			if (RMAX_ERR_BUSY == Status) 
+			{
+				FPlatformProcess::SleepNoStats(0.01f);
+			}
+		} while (Status == RMAX_ERR_BUSY);
+	}
+
+	void FRivermaxOutputStream::WaitForNextRound(double NextRoundTimeNs)
+	{
+		double NowTimeNs = FPlatformTime::Seconds() * 1E9;
+
+		//If we are already late, no need to wait
+		if (NextRoundTimeNs <= NowTimeNs)
+		{
+			return;
+		}
+
+		const double SleepTimeNs = NextRoundTimeNs - NowTimeNs;
+
+		static const bool bShouldPoll = false;
+
+		//Todo Should we have a threshold where we don't sleep?
+		//Todo Should we sleep less when we had retries
+		if (bShouldPoll)
+		{
+			while (NowTimeNs < SleepTimeNs)
+			{
+				NowTimeNs = FPlatformTime::Seconds() * 1E9;
+			}
+		}
+		else
+		{
+			const double SleepTimeS = SleepTimeNs / 1E9;
+			UE_LOG(LogRivermax, Verbose, TEXT("Outputstream sleeping for %0.5f"), SleepTimeS);
+			FPlatformProcess::SleepNoStats(SleepTimeS);
+		}
+	}
+
+	void FRivermaxOutputStream::GetNextChunk()
+	{
+		rmax_status_t Status;
+		do
+		{
+			Status = rmax_out_get_next_chunk(StreamId, &CurrentFrame->PayloadPtr, &CurrentFrame->HeaderPtr);
+			if (Status == RMAX_OK)
+			{
+				if (StreamData.bHasFrameFirstChunkBeenFetched == false)
+				{
+					StreamData.bHasFrameFirstChunkBeenFetched = true;
+					if (CurrentFrame->VideoBuffer != CurrentFrame->PayloadPtr)
+					{
+						//Debug code to track rivermax chunk processing
+						UE_LOG(LogRivermax, Warning, TEXT("Frame being sent (%d) doesn't match chunks being processed."), CurrentFrame->FrameIndex);
+					}
+				}
+
+				break;
+			}
+			else if (Status == RMAX_ERR_NO_FREE_CHUNK)
+			{
+				//We should not be here
+				Stats.ChunkRetries++;
+			}
+			else
+			{
+				UE_LOG(LogRivermax, Error, TEXT("Invalid error happened while trying to get next chunks. Status: %d"), Status);
+				Listener->OnStreamError();
+				bIsActive = false;
+			}
+		} while (Status != RMAX_OK && bIsActive);
+	}
+
+	void FRivermaxOutputStream::SetupRTPHeaders()
+	{
+		uint8* HeaderRawPtr = reinterpret_cast<uint8*>(CurrentFrame->HeaderPtr);
+		check(HeaderRawPtr);
+		for (uint32 StrideIndex = 0; StrideIndex < StreamMemory.ChunkSizeInStrides && CurrentFrame->PacketCounter < StreamMemory.PacketsInFrameField; ++StrideIndex)
+		{
+			uint8* NextHeaderRawPtr = HeaderRawPtr + (StrideIndex * StreamMemory.HeaderStrideSize);
+			BuildRTPHeader(*reinterpret_cast<FRTPHeader*>(NextHeaderRawPtr));
+			//todo only for video
+			if (!((StrideIndex + 1) % StreamMemory.PacketsInLine))
+			{
+				CurrentFrame->LineNumber = (CurrentFrame->LineNumber + 1) % Options.Resolution.Y; //preparing line number for next iteration
+			}
+			++StreamData.SequenceNumber;
+		}
+	}
+
+	void FRivermaxOutputStream::CommitNextChunks()
+	{
+		rmax_status_t Status;
+		do
+		{
+			// todo configure timeout
+			
+			const uint64 Timeout = 0;
+			const rmax_commit_flags_t CommitFlags{};
+			Status = rmax_out_commit(StreamId, Timeout, CommitFlags);
+			if (Status == RMAX_OK)
+			{
+				break;
+			}
+			else if (Status == RMAX_ERR_HW_SEND_QUEUE_FULL)
+			{
+				Stats.CommitRetries++;
+			}
+			else if(Status == RMAX_ERR_HW_COMPLETION_ISSUE)
+			{
+				UE_LOG(LogRivermax, Error, TEXT("Completion issue while trying to commit next round of chunks."));
+				Listener->OnStreamError();
+				bIsActive = false;
+			}
+
+		} while (Status != RMAX_OK && bIsActive);
+	}
+
+	bool FRivermaxOutputStream::Init()
+	{
+		return true;
+	}
+
+	uint32 FRivermaxOutputStream::Run()
+	{
+		while (bIsActive)
+		{
+			if (CurrentFrame.IsValid() == false)
+			{
+				CurrentFrame = GetNextFrameToSend();
+			}
+
+			if(CurrentFrame.IsValid())
+			{
+				Process_AnyThread();
+			}
+			else
+			{
+				TRACE_CPUPROFILER_EVENT_SCOPE(FRivermaxOutputStream::WaitForFrame);
+				const double PreWaitTime = FPlatformTime::Seconds();
+				ReadyToSendEvent->Wait();
+				UE_LOG(LogRivermax, Verbose, TEXT("Outputstream waited %0.8f"), FPlatformTime::Seconds() - PreWaitTime);
+			}
+		}
+
+
+		DestroyStream();
+
+		return 0;
+	}
+
+	void FRivermaxOutputStream::Stop()
+	{
+		bIsActive = false;
+	}
+
+	void FRivermaxOutputStream::Exit()
+	{
+
+	}
+}
+
diff --git a/Engine/Plugins/VirtualProduction/Rivermax/RivermaxCore/Source/Private/RivermaxOutputStream.h b/Engine/Plugins/VirtualProduction/Rivermax/RivermaxCore/Source/Private/RivermaxOutputStream.h
new file mode 100644
--- /dev/null
+++ b/Engine/Plugins/VirtualProduction/Rivermax/RivermaxCore/Source/Private/RivermaxOutputStream.h
@@ -0,0 +1,137 @@
+// Copyright Epic Games, Inc. All Rights Reserved.
+
+#pragma once
+
+#include "IRivermaxOutputStream.h"
+
+#include "HAL/Runnable.h"
+#include "RivermaxHeader.h"
+#include "RivermaxOutputFrame.h"
+#include "RivermaxTypes.h"
+
+
+class FEvent;
+
+namespace UE::RivermaxCore::Private
+{
+	using UE::RivermaxCore::FRivermaxStreamOptions;
+	using UE::RivermaxCore::FRivermaxOutputVideoFrameInfo;
+	
+	// Todo make a proper RTP header struct
+	struct FRTPHeader
+	{
+		uint8 RawHeader[20];
+	};
+
+	struct FRivermaxOutputStreamMemory
+	{
+		uint16 PayloadSize = 0; //todo
+		uint32 DataStrideSize = 1280; //todo align cache line?
+		uint32 HeaderStrideSize = 20; //todo align cache line?
+		uint32 LinesInChunk = 4;
+
+		uint32 PacketsInLine = 0;
+		uint32 ChunkSizeInStrides = LinesInChunk * PacketsInLine;
+
+		uint32 FramesFieldPerMemoryBlock = 1;
+		uint32 PacketsInFrameField = 0;
+		uint32 PacketsPerMemoryBlock = PacketsInFrameField * FramesFieldPerMemoryBlock;
+		uint32 ChunksPerFrameField = FMath::CeilToInt32(PacketsInFrameField / (double)ChunkSizeInStrides);
+		uint32 ChunksPerMemoryBlock = FramesFieldPerMemoryBlock * ChunksPerFrameField;
+		uint32 MemoryBlockCount = 10; //todo
+		uint32 StridesPerMemoryBlock = ChunkSizeInStrides * ChunksPerMemoryBlock;
+
+		TArray<rmax_mem_block> MemoryBlocks;
+		TArray<uint16_t> PayloadSizes; //Array describing stride payload size
+		TArray<uint16_t> HeaderSizes; //Array describing header payload size
+		TArray<FRTPHeader> RTPHeaders;
+
+		rmax_buffer_attr BufferAttributes;
+	};
+
+	struct FRivermaxOutputStreamStats
+	{
+		uint32 ChunkRetries = 0;
+		uint32 TotalStrides = 0;
+		uint32 ChunkWait = 0;
+		uint32 CommitWaits = 0;
+		uint32 CommitRetries = 0;
+		uint64 MemoryBlockSentCounter = 0; //Global for an active capture session to track timestamp for next packet
+	};
+
+	struct FRivermaxOutputStreamData
+	{
+		/** Current sequence number being done */
+		uint32 SequenceNumber = 0;
+		double FrameFieldTimeIntervalNs = 0.0;
+		double StartSendTimeNs = 0.0;
+		double SendTimeNs = 0.0;
+		double InitialTimestampTick = 0.0;
+		bool bHasFrameFirstChunkBeenFetched = false;
+	};
+
+
+	class FRivermaxOutputStream : public UE::RivermaxCore::IRivermaxOutputStream, public FRunnable
+	{
+	public:
+		FRivermaxOutputStream();
+		virtual ~FRivermaxOutputStream();
+
+	public:
+
+		//~ Begin IRivermaxOutputStream interface
+		virtual bool Initialize(const FRivermaxStreamOptions& Options, IRivermaxOutputStreamListener& InListener) override;
+		virtual void Uninitialize() override;
+		virtual bool PushVideoFrame(const FRivermaxOutputVideoFrameInfo& NewFrame) override;
+		//~ End IRivermaxOutputStream interface
+
+		void Process_AnyThread();
+
+		//~ Begin FRunnable interface
+		virtual bool Init() override;
+		virtual uint32 Run() override;
+		virtual void Stop() override;
+		virtual void Exit() override;
+		//~ End FRunnable interface
+
+	private:
+		void InitializeBuffers();
+		void InitializeMemory();
+		void InitializeNextFrame(const TSharedPtr<FRivermaxOutputFrame>& NextFrame);
+		TSharedPtr<FRivermaxOutputFrame> GetNextFrameToSend();
+		TSharedPtr<FRivermaxOutputFrame> GetNextAvailableFrame(uint32 InFrameIdentifier);
+		void BuildRTPHeader(FRTPHeader& OutHeader) const;
+		void DestroyStream();
+		void WaitForNextRound(double NextRoundTime);
+		void GetNextChunk();
+		void SetupRTPHeaders();
+		void CommitNextChunks();
+
+	private:
+		FRivermaxStreamOptions Options;
+		FRivermaxOutputStreamMemory StreamMemory;
+		FRivermaxOutputStreamStats Stats;
+		FRivermaxOutputStreamData StreamData;
+
+		rmax_stream_id StreamId;
+		FCriticalSection FrameCriticalSection;
+
+		TSharedPtr<FRivermaxOutputFrame> CurrentFrame;
+
+		TArray<TSharedPtr<FRivermaxOutputFrame>> AvailableFrames;
+		TArray<TSharedPtr<FRivermaxOutputFrame>> FramesToSend;
+		double LastTime = 0.0;
+
+		TUniquePtr<FRunnableThread> RivermaxThread;
+		std::atomic<bool> bIsActive;
+
+		FEvent* ReadyToSendEvent = nullptr;
+
+		IRivermaxOutputStreamListener* Listener = nullptr;
+
+		static constexpr double MediaClockSampleRate = 90000.0; //Required to comply with SMTPE 2110-10.The Media Clock and RTP Clock rate for streams compliant to this standard shall be 90 kHz.
+		ERivermaxStreamType StreamType = ERivermaxStreamType::VIDEO_2110_20_STREAM; //todo
+	};
+}
+
+
diff --git a/Engine/Plugins/VirtualProduction/Rivermax/RivermaxCore/Source/Private/RivermaxUtils.cpp b/Engine/Plugins/VirtualProduction/Rivermax/RivermaxCore/Source/Private/RivermaxUtils.cpp
new file mode 100644
--- /dev/null
+++ b/Engine/Plugins/VirtualProduction/Rivermax/RivermaxCore/Source/Private/RivermaxUtils.cpp
@@ -0,0 +1,36 @@
+// Copyright Epic Games, Inc. All Rights Reserved.
+
+#include "RivermaxUtils.h"
+
+#include "RivermaxTypes.h"
+
+
+namespace UE::RivermaxCore::Private::Utils
+{
+	void StreamOptionsToSDPDescription(const UE::RivermaxCore::FRivermaxStreamOptions& Options, FAnsiStringBuilderBase& OutSDPDescription)
+	{
+		// Basic SDP string creation from a set of options. At some point, having a proper SDP loader / creator.
+		// Refer to https://datatracker.ietf.org/doc/html/rfc4570
+
+		FString FrameRateDescription;
+		if (FMath::IsNearlyZero(FMath::Frac(Options.FrameRate.AsDecimal())) == false)
+		{
+			FrameRateDescription = FString::Printf(TEXT("%d/%d"), Options.FrameRate.Numerator, Options.FrameRate.Denominator);
+		}
+		else
+		{
+			FrameRateDescription = FString::Printf(TEXT("%d"), (uint32)Options.FrameRate.AsDecimal());
+		}
+
+		OutSDPDescription.Appendf("v=0\n");
+		OutSDPDescription.Appendf("s=SMPTE ST2110 20 streams\n");
+		OutSDPDescription.Appendf("m=video %d RTP/AVP 96\n", Options.Port);
+		OutSDPDescription.Appendf("c=IN IP4 %S/64\n", *Options.DestinationAddress);
+		OutSDPDescription.Appendf("a=source-filter: incl IN IP4 %S %S\n", *Options.DestinationAddress, *Options.SourceAddress);
+		OutSDPDescription.Appendf("a=rtpmap:96 raw/90000\n");
+		OutSDPDescription.Appendf("a=fmtp: 96 sampling=YCbCr-4:2:2; width=%d; height=%d; exactframerate=%S; depth=%d; TCS=SDR; colorimetry=BT709; PM=2110GPM; SSN=ST2110-20:2017; TP=2110TPN;\n"
+			, Options.Resolution.X, Options.Resolution.Y, *FrameRateDescription, Options.PixelFormat == ERivermaxOutputPixelFormat::RMAX_10BIT_YCBCR ? 10 : 8);
+		OutSDPDescription.Appendf("a=mediaclk:direct=0\n");
+		OutSDPDescription.Appendf("a=mid:VID");
+	}
+}
\ No newline at end of file
diff --git a/Engine/Plugins/VirtualProduction/Rivermax/RivermaxCore/Source/Private/RivermaxUtils.h b/Engine/Plugins/VirtualProduction/Rivermax/RivermaxCore/Source/Private/RivermaxUtils.h
new file mode 100644
--- /dev/null
+++ b/Engine/Plugins/VirtualProduction/Rivermax/RivermaxCore/Source/Private/RivermaxUtils.h
@@ -0,0 +1,26 @@
+// Copyright Epic Games, Inc. All Rights Reserved.
+
+#pragma once
+
+#include "CoreMinimal.h"
+
+namespace UE::RivermaxCore
+{
+	struct FRivermaxStreamOptions;
+}
+
+namespace UE::RivermaxCore::Private::Utils
+{	
+	/** Various constants used for stream initialization */
+	static constexpr uint32 BytesPerHeader = 20;
+	static constexpr uint32 BytesPerPacket = 1200;
+	static constexpr uint32 PixelsPerGroup_422_8b = 2;
+	static constexpr uint32 BytesPerGroup_422_8b = 4;
+	static constexpr uint32 FullHDHeight = 1080;
+	static constexpr uint32 FullHDWidth = 1920;
+	static constexpr uint32 MaxPayloadSize = 1280;
+	static constexpr uint32 VideoTROModification = 0; // Todo : Investigate what TRO modification is
+
+	/** Convert a set of streaming option to its SDP description. Currently only support video type. */
+	void StreamOptionsToSDPDescription(const UE::RivermaxCore::FRivermaxStreamOptions& Options, FAnsiStringBuilderBase& OutSDPDescription);
+}
diff --git a/Engine/Plugins/VirtualProduction/Rivermax/RivermaxCore/Source/Public/IRivermaxCoreModule.h b/Engine/Plugins/VirtualProduction/Rivermax/RivermaxCore/Source/Public/IRivermaxCoreModule.h
new file mode 100644
--- /dev/null
+++ b/Engine/Plugins/VirtualProduction/Rivermax/RivermaxCore/Source/Public/IRivermaxCoreModule.h
@@ -0,0 +1,44 @@
+// Copyright Epic Games, Inc. All Rights Reserved.
+
+#pragma once
+
+#include "Modules/ModuleInterface.h"
+
+#include "IRivermaxInputStream.h"
+#include "IRivermaxOutputStream.h"
+#include "Modules/ModuleManager.h"
+
+
+namespace UE::RivermaxCore
+{
+	class IRivermaxManager;
+}
+
+
+/**
+ * Core module for Rivermax access from the engine. Users can create different stream types that are exposed to 
+ * get data flow ongoing.
+ */
+class IRivermaxCoreModule : public IModuleInterface
+{
+public:
+	static inline IRivermaxCoreModule& Get()
+	{
+		return FModuleManager::LoadModuleChecked<IRivermaxCoreModule>("RivermaxCore");
+	}
+
+	static inline bool IsAvailable()
+	{
+		return FModuleManager::Get().IsModuleLoaded("RivermaxCore");
+	}
+
+	/** Create input stream managing receiving data from rivermax */
+	virtual TUniquePtr<UE::RivermaxCore::IRivermaxInputStream> CreateInputStream() = 0;
+
+	/** Create output stream managing sending data to rivermax */
+	virtual TUniquePtr<UE::RivermaxCore::IRivermaxOutputStream> CreateOutputStream() = 0;
+
+	/** Returns Rivermax manager singleton to query for stats, status, etc... */
+	virtual TSharedPtr<UE::RivermaxCore::IRivermaxManager> GetRivermaxManager() = 0;
+};
+
diff --git a/Engine/Plugins/VirtualProduction/Rivermax/RivermaxCore/Source/Public/IRivermaxInputStream.h b/Engine/Plugins/VirtualProduction/Rivermax/RivermaxCore/Source/Public/IRivermaxInputStream.h
new file mode 100644
--- /dev/null
+++ b/Engine/Plugins/VirtualProduction/Rivermax/RivermaxCore/Source/Public/IRivermaxInputStream.h
@@ -0,0 +1,52 @@
+// Copyright Epic Games, Inc. All Rights Reserved.
+
+#pragma once
+
+#include "CoreMinimal.h"
+
+namespace UE::RivermaxCore
+{
+	struct FRivermaxStreamOptions;
+
+	struct RIVERMAXCORE_API FRivermaxInputVideoFrameDescriptor
+	{
+		uint32 Height = 0;
+		uint32 Width = 0;
+		uint32 Stride = 0;
+		uint32 VideoBufferSize = 0;
+	};
+
+	struct RIVERMAXCORE_API FRivermaxInputVideoFrameRequest
+	{
+		uint8* VideoBuffer = nullptr;
+	};
+
+	struct RIVERMAXCORE_API FRivermaxInputVideoFrameReception
+	{
+		uint8* VideoBuffer = nullptr;
+	};
+
+	class RIVERMAXCORE_API IRivermaxInputStreamListener
+	{
+	public:
+		/** Initialization completion callback with result */
+		virtual void OnInitializationCompleted(bool bHasSucceed) = 0;
+	
+		/** Called when stream is ready to fill the next frame. Returns true if a frame was successfully requested */
+		virtual bool OnVideoFrameRequested(const FRivermaxInputVideoFrameDescriptor& FrameInfo, FRivermaxInputVideoFrameRequest& OutVideoFrameRequest) = 0;
+		
+		/** Called when a frame has been received */
+		virtual void OnVideoFrameReceived(const FRivermaxInputVideoFrameDescriptor& FrameInfo, const FRivermaxInputVideoFrameReception& ReceivedVideoFrame) = 0;
+	};
+
+	class RIVERMAXCORE_API IRivermaxInputStream
+	{
+	public:
+		virtual ~IRivermaxInputStream() = default;
+
+	public:
+		virtual bool Initialize(const FRivermaxStreamOptions& InOptions, IRivermaxInputStreamListener& InListener) = 0;
+		virtual void Uninitialize() = 0;
+	};
+}
+
diff --git a/Engine/Plugins/VirtualProduction/Rivermax/RivermaxCore/Source/Public/IRivermaxManager.h b/Engine/Plugins/VirtualProduction/Rivermax/RivermaxCore/Source/Public/IRivermaxManager.h
new file mode 100644
--- /dev/null
+++ b/Engine/Plugins/VirtualProduction/Rivermax/RivermaxCore/Source/Public/IRivermaxManager.h
@@ -0,0 +1,24 @@
+// Copyright Epic Games, Inc. All Rights Reserved.
+
+#pragma once
+
+#include "CoreMinimal.h"
+
+namespace UE::RivermaxCore
+{
+	/**
+	 * Not doing much at the moment but should be used as a central point to register every stream created
+	 * and keep track of stats. Also manages initialization of the library.
+	 */
+	class RIVERMAXCORE_API IRivermaxManager
+	{
+	public:
+		virtual ~IRivermaxManager() = default;
+
+	public:
+		
+		/** Returns true if Rivermax has been initialized and is usable */
+		virtual bool IsInitialized() const = 0;
+	};
+}
+
diff --git a/Engine/Plugins/VirtualProduction/Rivermax/RivermaxCore/Source/Public/IRivermaxOutputStream.h b/Engine/Plugins/VirtualProduction/Rivermax/RivermaxCore/Source/Public/IRivermaxOutputStream.h
new file mode 100644
--- /dev/null
+++ b/Engine/Plugins/VirtualProduction/Rivermax/RivermaxCore/Source/Public/IRivermaxOutputStream.h
@@ -0,0 +1,53 @@
+// Copyright Epic Games, Inc. All Rights Reserved.
+
+#pragma once
+
+#include "CoreMinimal.h"
+
+
+namespace UE::RivermaxCore
+{
+	struct FRivermaxStreamOptions;
+
+	struct RIVERMAXCORE_API FRivermaxOutputVideoFrameInfo
+	{
+		uint32 FrameIdentifier = 0;
+		uint32 Height = 0;
+		uint32 Width = 0;
+		uint32 Stride = 0;
+		void* VideoBuffer = nullptr;
+	};
+
+	class RIVERMAXCORE_API IRivermaxOutputStreamListener
+	{
+	public:
+
+		/** Initialization completion callback with result */
+		virtual void OnInitializationCompleted(bool bHasSucceed) = 0;
+
+		/** Called when stream has encountered an error and has to stop */
+		virtual void OnStreamError() = 0;
+	};
+
+	class RIVERMAXCORE_API IRivermaxOutputStream
+	{
+	public:
+		virtual ~IRivermaxOutputStream() = default;
+
+	public:
+
+		/** 
+		 * Initializes stream using input options. Returns false if stream creation has failed. 
+		 * Initialization is completed when OnInitializationCompletion has been called with result.
+		 */
+		virtual bool Initialize(const FRivermaxStreamOptions& Options, IRivermaxOutputStreamListener& InListener) = 0;
+
+		/** Uninitializes current stream */
+		virtual void Uninitialize() = 0;
+
+		/** Pushes new video frame to the stream. Returns false if frame couldn't be pushed. */
+		virtual bool PushVideoFrame(const FRivermaxOutputVideoFrameInfo& NewFrame) = 0;
+	};
+}
+
+
diff --git a/Engine/Plugins/VirtualProduction/Rivermax/RivermaxCore/Source/Public/RivermaxHeader.h b/Engine/Plugins/VirtualProduction/Rivermax/RivermaxCore/Source/Public/RivermaxHeader.h
new file mode 100644
--- /dev/null
+++ b/Engine/Plugins/VirtualProduction/Rivermax/RivermaxCore/Source/Public/RivermaxHeader.h
@@ -0,0 +1,20 @@
+// Copyright Epic Games, Inc. All Rights Reserved.
+
+#pragma once
+
+THIRD_PARTY_INCLUDES_START
+
+#if PLATFORM_WINDOWS
+#include "Windows/AllowWindowsPlatformTypes.h"
+#include "Windows/PreWindowsApi.h"
+#endif
+
+#include "rivermax_api.h"
+
+#if PLATFORM_WINDOWS
+#include "Windows/PostWindowsApi.h"
+#include "Windows/HideWindowsPlatformTypes.h"
+#endif
+
+THIRD_PARTY_INCLUDES_END
+
diff --git a/Engine/Plugins/VirtualProduction/Rivermax/RivermaxCore/Source/Public/RivermaxTypes.h b/Engine/Plugins/VirtualProduction/Rivermax/RivermaxCore/Source/Public/RivermaxTypes.h
new file mode 100644
--- /dev/null
+++ b/Engine/Plugins/VirtualProduction/Rivermax/RivermaxCore/Source/Public/RivermaxTypes.h
@@ -0,0 +1,54 @@
+// Copyright Epic Games, Inc. All Rights Reserved.
+
+#pragma once
+
+#include "CoreMinimal.h"
+
+#include "Misc/FrameRate.h"
+
+namespace UE::RivermaxCore
+{
+	enum class RIVERMAXCORE_API ERivermaxOutputPixelFormat
+	{
+		RMAX_8BIT_YCBCR,
+		RMAX_8BIT_RGB,
+		RMAX_10BIT_YCBCR,
+		RMAX_10BIT_RGB,
+	};
+
+	struct RIVERMAXCORE_API FRivermaxStreamOptions
+	{
+		/** Desired stream resolution */
+		FIntPoint Resolution = { 1920, 1080 };
+
+		/** Stream FrameRate */
+		FFrameRate FrameRate = { 24,1 };
+
+		/** Source IP to bind to */
+		FString SourceAddress;
+
+		/** Destination IP to send to. Defaults to multicast group IP. */
+		FString DestinationAddress = TEXT("224.1.1.1");
+
+		/** Port to be used by stream */
+		uint32 Port = 50000;
+
+		/** Desired stream pixel format */
+		ERivermaxOutputPixelFormat PixelFormat = ERivermaxOutputPixelFormat::RMAX_8BIT_YCBCR;
+
+		/** Sample count to buffer. */
+		int32 NumberOfBuffers = 2;
+	};
+
+	enum class RIVERMAXCORE_API ERivermaxStreamType : uint8
+	{
+		VIDEO_2110_20_STREAM,
+
+		/** Todo add additional stream types */
+		//VIDEO_2110_22_STREAM, compressed video
+		//AUDIO_2110_30_31_STREAM,
+		//ANCILLARY_2110_40_STREAM,
+	};
+}
+
+
diff --git a/Engine/Plugins/VirtualProduction/Rivermax/RivermaxCore/Source/RivermaxCore.build.cs b/Engine/Plugins/VirtualProduction/Rivermax/RivermaxCore/Source/RivermaxCore.build.cs
new file mode 100644
--- /dev/null
+++ b/Engine/Plugins/VirtualProduction/Rivermax/RivermaxCore/Source/RivermaxCore.build.cs
@@ -0,0 +1,41 @@
+// Copyright Epic Games, Inc. All Rights Reserved.
+
+using UnrealBuildTool;
+
+public class RivermaxCore : ModuleRules
+{
+	public RivermaxCore(ReadOnlyTargetRules Target) : base(Target)
+	{
+		PublicDependencyModuleNames.AddRange(
+			new string[] {
+				"Core",
+				"CoreUObject",
+				"MediaAssets",
+				"MediaIOCore",
+				"RivermaxLib",
+				"TimeManagement"
+			}
+		);
+
+		PrivateDependencyModuleNames.AddRange(
+			new string[] {
+				"Engine",
+				"RenderCore",
+				"RHI"
+			}
+		);
+
+		if (Target.bBuildEditor)
+		{
+			PublicDependencyModuleNames.AddRange(
+				new string[] {
+				}
+			);
+
+			PrivateDependencyModuleNames.AddRange(
+				new string[] {
+				}
+			);
+		}
+	}
+}
