- Fixing static analysis warning for rivermax output

diff --git a/Engine/Plugins/VirtualProduction/Rivermax/RivermaxCore/Source/Private/RivermaxOutputStream.h b/Engine/Plugins/VirtualProduction/Rivermax/RivermaxCore/Source/Private/RivermaxOutputStream.h
--- a/Engine/Plugins/VirtualProduction/Rivermax/RivermaxCore/Source/Private/RivermaxOutputStream.h
+++ b/Engine/Plugins/VirtualProduction/Rivermax/RivermaxCore/Source/Private/RivermaxOutputStream.h
@@ -25,21 +25,21 @@ namespace UE::RivermaxCore::Private
 
 	struct FRivermaxOutputStreamMemory
 	{
-		uint16 PayloadSize = 0; //todo
-		uint32 DataStrideSize = 1280; //todo align cache line?
-		uint32 HeaderStrideSize = 20; //todo align cache line?
+		uint16 PayloadSize = 0; 
+		uint32 DataStrideSize = 1280; 
+		uint32 HeaderStrideSize = 20;
 		uint32 LinesInChunk = 4;
 
 		uint32 PacketsInLine = 0;
-		uint32 ChunkSizeInStrides = LinesInChunk * PacketsInLine;
+		uint32 ChunkSizeInStrides = 0;
 
 		uint32 FramesFieldPerMemoryBlock = 1;
 		uint32 PacketsInFrameField = 0;
-		uint32 PacketsPerMemoryBlock = PacketsInFrameField * FramesFieldPerMemoryBlock;
-		uint32 ChunksPerFrameField = FMath::CeilToInt32(PacketsInFrameField / (double)ChunkSizeInStrides);
-		uint32 ChunksPerMemoryBlock = FramesFieldPerMemoryBlock * ChunksPerFrameField;
-		uint32 MemoryBlockCount = 10; //todo
-		uint32 StridesPerMemoryBlock = ChunkSizeInStrides * ChunksPerMemoryBlock;
+		uint32 PacketsPerMemoryBlock = 0;
+		uint32 ChunksPerFrameField = 0;
+		uint32 ChunksPerMemoryBlock = 0;
+		uint32 MemoryBlockCount = 0; 
+		uint32 StridesPerMemoryBlock = 0;
 
 		TArray<rmax_mem_block> MemoryBlocks;
 		TArray<uint16_t> PayloadSizes; //Array describing stride payload size
