- Quick fix for rivermax output having visible tearing / buffer corruption

diff --git a/Engine/Plugins/VirtualProduction/Rivermax/RivermaxCore/Source/RivermaxCore/Private/RivermaxOutputStream.cpp b/Engine/Plugins/VirtualProduction/Rivermax/RivermaxCore/Source/RivermaxCore/Private/RivermaxOutputStream.cpp
--- a/Engine/Plugins/VirtualProduction/Rivermax/RivermaxCore/Source/RivermaxCore/Private/RivermaxOutputStream.cpp
+++ b/Engine/Plugins/VirtualProduction/Rivermax/RivermaxCore/Source/RivermaxCore/Private/RivermaxOutputStream.cpp
@@ -256,7 +256,13 @@ namespace UE::RivermaxCore::Private
 
 	void FRivermaxOutputStream::Process_AnyThread()
 	{
-		if (CurrentFrame.IsValid())
+		// Not the best path but it seems to work without tearing for now
+		// Wait for the next time a frame should be sent (based on frame interval)
+		// If a frame had been sent, make it available again
+		// This is to avoid making it available right after scheduling it. It's not sent yet and we start overwriting it
+		// Get next available frame that was rendered (Wait if none are)
+		// Send frame
+		// Restart
 		{
 			{
 				TRACE_CPUPROFILER_EVENT_SCOPE(FRivermaxOutputStream::Wait);
@@ -264,6 +270,36 @@ namespace UE::RivermaxCore::Private
 				WaitForNextRound(StreamData.SendTimeNs);
 			}
 
+			{
+				if (CurrentFrame.IsValid())
+				{
+					//Release last frame sent. We keep hold to avoid overwriting it as rivermax is sending it
+					CurrentFrame->Clear();
+					{
+						FScopeLock Lock(&FrameCriticalSection);
+						AvailableFrames.Add(MoveTemp(CurrentFrame));
+						CurrentFrame.Reset();
+					}
+				}
+
+				while (CurrentFrame.IsValid() == false && bIsActive)
+				{
+					CurrentFrame = GetNextFrameToSend();
+					if (!CurrentFrame.IsValid())
+					{
+						TRACE_CPUPROFILER_EVENT_SCOPE(FRivermaxOutputStream::WaitForFrame);
+						const double PreWaitTime = FPlatformTime::Seconds();
+						ReadyToSendEvent->Wait();
+						UE_LOG(LogRivermax, Verbose, TEXT("Outputstream waited %0.8f"), FPlatformTime::Seconds() - PreWaitTime);
+					}
+				}
+			}
+
+			if (bIsActive == false)
+			{
+				return;
+			}
+
 			const FString TraceName = FString::Format(TEXT("FRivermaxOutputStream::SendFrame {0}"), { CurrentFrame->FrameIndex });
 			TRACE_CPUPROFILER_EVENT_SCOPE_TEXT(*TraceName);
 
@@ -297,13 +333,6 @@ namespace UE::RivermaxCore::Private
 
 			} while (CurrentFrame->ChunkNumber < StreamMemory.ChunksPerMemoryBlock && bIsActive);
 			
-			//We are done with this frame, put it back as available
-			CurrentFrame->Clear();
-			{
-				FScopeLock Lock(&FrameCriticalSection);
-				AvailableFrames.Add(MoveTemp(CurrentFrame));
-				CurrentFrame.Reset();
-			}
 
 			Stats.MemoryBlockSentCounter++;
 			Stats.TotalStrides += StreamMemory.ChunkSizeInStrides;
@@ -645,22 +674,7 @@ namespace UE::RivermaxCore::Private
 	{
 		while (bIsActive)
 		{
-			if (CurrentFrame.IsValid() == false)
-			{
-				CurrentFrame = GetNextFrameToSend();
-			}
-
-			if(CurrentFrame.IsValid())
-			{
-				Process_AnyThread();
-			}
-			else
-			{
-				TRACE_CPUPROFILER_EVENT_SCOPE(FRivermaxOutputStream::WaitForFrame);
-				const double PreWaitTime = FPlatformTime::Seconds();
-				ReadyToSendEvent->Wait();
-				UE_LOG(LogRivermax, Verbose, TEXT("Outputstream waited %0.8f"), FPlatformTime::Seconds() - PreWaitTime);
-			}
+			Process_AnyThread();
 		}
 
 
