- Rivermax player : Adding support for 2110 RGB 8 and 10 bit input in engine

diff --git a/Engine/Plugins/VirtualProduction/Rivermax/RivermaxCore/Shaders/Private/RivermaxShaders.usf b/Engine/Plugins/VirtualProduction/Rivermax/RivermaxCore/Shaders/Private/RivermaxShaders.usf
--- a/Engine/Plugins/VirtualProduction/Rivermax/RivermaxCore/Shaders/Private/RivermaxShaders.usf
+++ b/Engine/Plugins/VirtualProduction/Rivermax/RivermaxCore/Shaders/Private/RivermaxShaders.usf
@@ -18,7 +18,7 @@ float InputPixelOffsetX;
 float InputPixelOffsetY;
 uint HorizontalElementCount;
 uint InputTexturePixelsPerThread;
-
+RWTexture2D<float4> OutTexture;
 
 struct FYCbCr10Bit422
 {
@@ -37,12 +37,14 @@ struct FRGB8BitBuffer
 	uint DWord2;
 };
 RWStructuredBuffer<FRGB8BitBuffer>	OutRGB8Buffer;
+StructuredBuffer<FRGB8BitBuffer> InputRGB8Buffer;
 
 struct FRGB10BitBuffer
 {
 	uint DWords[15];
 };
 RWStructuredBuffer<FRGB10BitBuffer>	OutRGB10Buffer;
+StructuredBuffer<FRGB10BitBuffer> InputBuffer;
 
 struct FRGB16fBuffer
 {
@@ -155,17 +157,17 @@ void RGBToRGB10BitCS(uint3 ThreadId : SV_DispatchThreadID, uint GroupIndex : SV_
 	OutRGB10Buffer[Index].DWords[3] =							  ((RGB04.x >> 2)   << 24)								| ((RGB03.z & 0xFF) << 16) | ((RGB03.y & 0x3F) << 10) | ((RGB03.z >> 8)   << 8)  | ((RGB03.x & 0xF)  << 4)  | (RGB03.y >> 6);
 	OutRGB10Buffer[Index].DWords[4] =							  ((RGB05.x & 0xFF) << 24)   | ((RGB04.z & 0x3F) << 18) | ((RGB05.x >> 8)   << 16) | ((RGB04.y & 0xF)  << 12) | ((RGB04.z >> 6)   << 8)  | ((RGB04.x & 0x3)  << 6)  | (RGB04.y >> 4);
 	
-	OutRGB10Buffer[Index].DWords[5] = ((RGB05.x & 0x03F) << 26) | ((RGB06.y >> 8)   << 24)   | ((RGB05.z & 0xF)  << 20) | ((RGB05.x >> 6)   << 16) | ((RGB05.y & 0x3)  << 14) | ((RGB05.z >> 4)   << 8)								| (RGB05.y >> 2);
-	OutRGB10Buffer[Index].DWords[6] = ((RGB06.x & 0x00F) << 28) | ((RGB07.y >> 6)   << 24)   | ((RGB06.z & 0x3)  << 22) | ((RGB06.x >> 4)   << 16)							  | ((RGB06.z >> 2)   << 8)								| (RGB06.y & 0xFF);
-	OutRGB10Buffer[Index].DWords[7] = ((RGB07.x & 0x003) << 30) | ((RGB08.y >> 4)   << 24)								| ((RGB07.x >> 2)   << 16)							  | ((RGB07.z & 0xFF) << 8)  | ((RGB07.y & 0x3F) << 2)  | (RGB07.z >> 8);
-	OutRGB10Buffer[Index].DWords[8] =							  ((RGB09.y >> 2)   << 24)								| ((RGB08.x & 0xFF) << 16) | ((RGB08.z & 0x3F) << 10) | ((RGB08.x >> 8)   << 8)  | ((RGB08.y & 0xF)  << 4)  | (RGB08.z >> 6);
-	OutRGB10Buffer[Index].DWords[9] =							  ((RGB10.y & 0xFF) << 24)   | ((RGB09.x & 0x3F) << 18) | ((RGB10.y >> 8)   << 16) | ((RGB09.z & 0xF)  << 12) | ((RGB09.x >> 6)   << 8)  | ((RGB09.y & 0x3)  << 6)  | (RGB09.z >> 4);
-
-	OutRGB10Buffer[Index].DWords[10] = ((RGB10.y & 0x03F) << 26) | ((RGB11.z >> 8)   << 24)  | ((RGB11.x & 0xF)  << 20) | ((RGB11.y >> 6)   << 16) | ((RGB10.z & 0x3)  << 14) | ((RGB10.x >> 4)   << 8)								| (RGB10.z >> 2);
-	OutRGB10Buffer[Index].DWords[11] = ((RGB11.y & 0x00F) << 28) | ((RGB12.z >> 6)   << 24)  | ((RGB12.x & 0x3)  << 22) | ((RGB12.y >> 4)   << 16)							  | ((RGB11.x >> 2)   << 8)								| (RGB11.z & 0xFF);
-	OutRGB10Buffer[Index].DWords[12] = ((RGB12.y & 0x003) << 30) | ((RGB13.z >> 4)   << 24)								| ((RGB13.y >> 2)   << 16)							  | ((RGB12.x & 0xFF) << 8)  | ((RGB12.z & 0x3F) << 2)  | (RGB12.x >> 8);
-	OutRGB10Buffer[Index].DWords[13] =							   ((RGB14.z >> 2)   << 24)								| ((RGB14.y & 0xFF) << 16) | ((RGB13.x & 0x3F) << 10) | ((RGB13.y >> 8)   << 8)  | ((RGB13.z & 0xF)  << 4)  | (RGB13.x >> 6);
-	OutRGB10Buffer[Index].DWords[14] =							   ((RGB15.z & 0xFF) << 24)  | ((RGB14.y & 0x3F) << 18) | ((RGB15.z >> 8)   << 16) | ((RGB14.x & 0xF)  << 12) | ((RGB14.y >> 6)   << 8)  | ((RGB14.z & 0x3)  << 6)  | (RGB14.x >> 4);
+	OutRGB10Buffer[Index].DWords[5] = ((RGB06.x & 0x03F) << 26) | ((RGB06.y >> 8)   << 24)   | ((RGB05.z & 0xF)  << 20) | ((RGB06.x >> 6)   << 16) | ((RGB05.y & 0x3)  << 14) | ((RGB05.z >> 4)   << 8)								| (RGB05.y >> 2);
+	OutRGB10Buffer[Index].DWords[6] = ((RGB07.x & 0x00F) << 28) | ((RGB07.y >> 6)   << 24)   | ((RGB06.z & 0x3)  << 22) | ((RGB07.x >> 4)   << 16)							  | ((RGB06.z >> 2)   << 8)								| (RGB06.y & 0xFF);
+	OutRGB10Buffer[Index].DWords[7] = ((RGB08.x & 0x003) << 30) | ((RGB08.y >> 4)   << 24)								| ((RGB08.x >> 2)   << 16)							  | ((RGB07.z & 0xFF) << 8)  | ((RGB07.y & 0x3F) << 2)  | (RGB07.z >> 8);
+	OutRGB10Buffer[Index].DWords[8] =							  ((RGB09.y >> 2)   << 24)								| ((RGB09.x & 0xFF) << 16) | ((RGB08.z & 0x3F) << 10) | ((RGB09.x >> 8)   << 8)  | ((RGB08.y & 0xF)  << 4)  | (RGB08.z >> 6);
+	OutRGB10Buffer[Index].DWords[9] =							  ((RGB10.y & 0xFF) << 24)   | ((RGB10.x & 0x3F) << 18) | ((RGB10.y >> 8)   << 16) | ((RGB09.z & 0xF)  << 12) | ((RGB10.x >> 6)   << 8)  | ((RGB09.y & 0x3)  << 6)  | (RGB09.z >> 4);
+
+	OutRGB10Buffer[Index].DWords[10] = ((RGB11.y & 0x03F) << 26) | ((RGB11.z >> 8)   << 24)  | ((RGB11.x & 0xF)  << 20) | ((RGB11.y >> 6)   << 16) | ((RGB10.z & 0x3)  << 14) | ((RGB11.x >> 4)   << 8)								| (RGB10.z >> 2);
+	OutRGB10Buffer[Index].DWords[11] = ((RGB12.y & 0x00F) << 28) | ((RGB12.z >> 6)   << 24)  | ((RGB12.x & 0x3)  << 22) | ((RGB12.y >> 4)   << 16)							  | ((RGB12.x >> 2)   << 8)								| (RGB11.z & 0xFF);
+	OutRGB10Buffer[Index].DWords[12] = ((RGB13.y & 0x003) << 30) | ((RGB13.z >> 4)   << 24)								| ((RGB13.y >> 2)   << 16)							  | ((RGB13.x & 0xFF) << 8)  | ((RGB12.z & 0x3F) << 2)  | (RGB13.x >> 8);
+	OutRGB10Buffer[Index].DWords[13] =							   ((RGB14.z >> 2)   << 24)								| ((RGB14.y & 0xFF) << 16) | ((RGB14.x & 0x3F) << 10) | ((RGB14.y >> 8)   << 8)  | ((RGB13.z & 0xF)  << 4)  | (RGB14.x >> 6);
+	OutRGB10Buffer[Index].DWords[14] =							   ((RGB15.z & 0xFF) << 24)  | ((RGB15.y & 0x3F) << 18) | ((RGB15.z >> 8)   << 16) | ((RGB15.x & 0xF)  << 12) | ((RGB15.y >> 6)   << 8)  | ((RGB14.z & 0x3)  << 6)  | (RGB15.x >> 4);
 }
 
 // Shader to convert from RGBA to RGB Float16
@@ -188,3 +190,206 @@ void RGBToRGB16fBitCS(uint3 ThreadId : SV_DispatchThreadID, uint GroupIndex : SV
 	OutRGB16fBuffer[Index].Pixel1[1]  = RGB1.y;
 	OutRGB16fBuffer[Index].Pixel1[2]  = RGB1.x;
 }
+
+// Shader to convert from packed RGB8 to RGBA8 texture
+
+[numthreads(8,8,1)]
+void RGB8BitToRGBA8CS(uint3 ThreadId : SV_DispatchThreadID, uint GroupIndex : SV_GroupIndex)
+{
+	// Pixels we treat in shader is pretty much hardcoded to 16 so no parameter has been used
+	const uint OutputTexturePixelsPerThread = 4;
+
+	// Scale U (x index) based on the number of pixels we treat per thread
+	const uint BaseU = ThreadId.x * OutputTexturePixelsPerThread;
+	const uint BaseV = ThreadId.y;
+
+	// Input buffer indexing is done in one dimension so adjust using thread Id and elements per row
+	const uint BufferIndex = ThreadId.y * HorizontalElementCount + ThreadId.x;
+
+	// Used to normalize pixel values based on 8 bits input
+	const float Mul = 1.0 / 255.0;
+
+	float3 RGB00;
+	RGB00.x = (((InputRGB8Buffer[BufferIndex].DWord0 >> 00) & 0xFF)) * Mul;
+	RGB00.y = (((InputRGB8Buffer[BufferIndex].DWord0 >>  8) & 0xFF)) * Mul;
+	RGB00.z = (((InputRGB8Buffer[BufferIndex].DWord0 >> 16) & 0xFF)) * Mul;
+
+	float3 RGB01;
+	RGB01.x = (((InputRGB8Buffer[BufferIndex].DWord0 >> 24) & 0xFF)) * Mul;
+	RGB01.y = (((InputRGB8Buffer[BufferIndex].DWord1 >> 00) & 0xFF)) * Mul;
+	RGB01.z = (((InputRGB8Buffer[BufferIndex].DWord1 >>  8) & 0xFF)) * Mul;
+
+	float3 RGB02;
+	RGB02.x = (((InputRGB8Buffer[BufferIndex].DWord1 >> 16) & 0xFF)) * Mul;
+	RGB02.y = (((InputRGB8Buffer[BufferIndex].DWord1 >> 24) & 0xFF)) * Mul;
+	RGB02.z = (((InputRGB8Buffer[BufferIndex].DWord2 >> 00) & 0xFF)) * Mul;
+	
+	float3 RGB03;
+	RGB03.x = (((InputRGB8Buffer[BufferIndex].DWord2 >>  8) & 0xFF)) * Mul;
+	RGB03.y = (((InputRGB8Buffer[BufferIndex].DWord2 >> 16) & 0xFF)) * Mul;
+	RGB03.z = (((InputRGB8Buffer[BufferIndex].DWord2 >> 24) & 0xFF)) * Mul;
+	
+
+	#if DO_SRGB_TO_LINEAR
+		RGB00 = sRGBToLinear(RGB00);
+		RGB01 = sRGBToLinear(RGB01);
+		RGB02 = sRGBToLinear(RGB02);
+		RGB03 = sRGBToLinear(RGB03);
+	#endif
+
+	OutTexture[uint2(BaseU + 0, BaseV)] = float4(RGB00, 1.0);
+	OutTexture[uint2(BaseU + 1, BaseV)] = float4(RGB01, 1.0);
+	OutTexture[uint2(BaseU + 2, BaseV)] = float4(RGB02, 1.0);
+	OutTexture[uint2(BaseU + 3, BaseV)] = float4(RGB03, 1.0);
+}
+
+// Shader to convert from packed RGB10 to RGB10A2 texture
+
+[numthreads(8,8,1)]
+void RGB10BitToRGBACS(uint3 ThreadId : SV_DispatchThreadID, uint GroupIndex : SV_GroupIndex)
+{
+	// Pixels we treat in shader is pretty much hardcoded to 16 so no parameter has been used
+	const uint OutputTexturePixelsPerThread = 16;
+
+	// Scale U (x index) based on the number of pixels we treat per thread
+	const uint BaseU = ThreadId.x * OutputTexturePixelsPerThread;
+	const uint BaseV = ThreadId.y;
+
+	// Input buffer indexing is done in one dimension so adjust using thread Id and elements per row
+	const uint BufferIndex = ThreadId.y * HorizontalElementCount + ThreadId.x;
+
+	// Used to normalize pixel values based on 10 bits input
+	const float Mul = 1.0 / 1023.0;
+
+	//InputBuffer[BufferIndex].DWords[0] = ((RGB00.z & 0x03F) << 26) | ((RGB01.x >> 8)   << 24)   | ((RGB00.y & 0xF)  << 20) | ((RGB00.z >> 6)   << 16) | ((RGB00.x & 0x3)  << 14) | ((RGB00.y >> 4)   << 8)								| (RGB00.x >> 2);
+
+	float3 RGB00 = uint3(0,0,0);
+	RGB00.x = ((((InputBuffer[BufferIndex].DWords[00] >>  0)  << 2) | ((InputBuffer[BufferIndex].DWords[00] >> 14) & 0x003)) & 0x3FF) * Mul;
+	RGB00.y = ((((InputBuffer[BufferIndex].DWords[00] >>  8)  << 4) | ((InputBuffer[BufferIndex].DWords[00] >> 20) & 0x00F)) & 0x3FF) * Mul;
+	RGB00.z = ((((InputBuffer[BufferIndex].DWords[00] >> 16)  << 6) | ((InputBuffer[BufferIndex].DWords[00] >> 26) & 0x03F)) & 0x3FF) * Mul;
+
+	//InputBuffer[BufferIndex].DWords[1] = ((RGB01.z & 0x00F) << 28) | ((RGB02.x >> 6)   << 24)   | ((RGB01.y & 0x3)  << 22) | ((RGB01.z >> 4)   << 16)							  | ((RGB01.y >> 2)   << 8)								| (RGB01.x & 0xFF);
+	float3 RGB01 = uint3(0,0,0);
+	RGB01.x = ((((InputBuffer[BufferIndex].DWords[00] >> 24) << 8) | ((InputBuffer[BufferIndex].DWords[01] >> 00) & 0x0FF)) & 0x3FF) * Mul;
+	RGB01.y = ((((InputBuffer[BufferIndex].DWords[01] >>  8) << 2) | ((InputBuffer[BufferIndex].DWords[01] >> 22) & 0x003)) & 0x3FF) * Mul;
+	RGB01.z = ((((InputBuffer[BufferIndex].DWords[01] >> 16) << 4) | ((InputBuffer[BufferIndex].DWords[01] >> 28) & 0x00F)) & 0x3FF) * Mul;
+
+	//InputBuffer[BufferIndex].DWords[2] = ((RGB02.z & 0x003) << 30) | ((RGB03.x >> 4)   << 24)								| ((RGB02.z >> 2)   << 16)							  | ((RGB02.y & 0xFF) << 8)  | ((RGB02.x & 0x3F) << 2)  | (RGB02.y >> 8);
+	float3 RGB02 = uint3(0,0,0);
+	RGB02.x = ((((InputBuffer[BufferIndex].DWords[01] >> 24) << 6) | ((InputBuffer[BufferIndex].DWords[02] >> 02) & 0x03F)) & 0x3FF) * Mul;
+	RGB02.y = ((((InputBuffer[BufferIndex].DWords[02] >>  0) << 8) | ((InputBuffer[BufferIndex].DWords[02] >>  8) & 0x0FF)) & 0x3FF) * Mul;
+	RGB02.z = ((((InputBuffer[BufferIndex].DWords[02] >> 16) << 2) | ((InputBuffer[BufferIndex].DWords[02] >> 30) & 0x002)) & 0x3FF) * Mul;
+	
+	//InputBuffer[BufferIndex].DWords[3] =							  ((RGB04.x >> 2)   << 24)								| ((RGB03.z & 0xFF) << 16) | ((RGB03.y & 0x3F) << 10) | ((RGB03.z >> 8)   << 8)  | ((RGB03.x & 0xF)  << 4)  | (RGB03.y >> 6);
+	float3 RGB03 = uint3(0,0,0);
+	RGB03.x = ((((InputBuffer[BufferIndex].DWords[02] >> 24) << 4) | ((InputBuffer[BufferIndex].DWords[03] >> 04) & 0x00F)) & 0x3FF) * Mul;
+	RGB03.y = ((((InputBuffer[BufferIndex].DWords[03] >> 00) << 6) | ((InputBuffer[BufferIndex].DWords[03] >> 10) & 0x03F)) & 0x3FF) * Mul;
+	RGB03.z = ((((InputBuffer[BufferIndex].DWords[03] >>  8) << 8) | ((InputBuffer[BufferIndex].DWords[03] >> 16) & 0x0FF)) & 0x3FF) * Mul;
+	
+	//InputBuffer[BufferIndex].DWords[4] =							  ((RGB05.x & 0xFF) << 24)   | ((RGB04.z & 0x3F) << 18) | ((RGB05.x >> 8)   << 16) | ((RGB04.y & 0xF)  << 12) | ((RGB04.z >> 6)   << 8)  | ((RGB04.x & 0x3)  << 6)  | (RGB04.y >> 4);
+	float3 RGB04 = uint3(0,0,0);
+	RGB04.x = ((((InputBuffer[BufferIndex].DWords[03] >> 24) << 2) | ((InputBuffer[BufferIndex].DWords[04] >> 06) & 0x003)) & 0x3FF) * Mul;
+	RGB04.y = ((((InputBuffer[BufferIndex].DWords[04] >> 00) << 4) | ((InputBuffer[BufferIndex].DWords[04] >> 12) & 0x00F)) & 0x3FF) * Mul;
+	RGB04.z = ((((InputBuffer[BufferIndex].DWords[04] >>  8) << 6) | ((InputBuffer[BufferIndex].DWords[04] >> 18) & 0x03F)) & 0x3FF) * Mul;
+	
+	//OutRGB10Buffer[Index].DWords[5] = ((RGB06.x & 0x03F) << 26) | ((RGB06.y >> 8)   << 24)   | ((RGB05.z & 0xF)  << 20) | ((RGB06.x >> 6)   << 16) | ((RGB05.y & 0x3)  << 14) | ((RGB05.z >> 4)   << 8)								| (RGB05.y >> 2);
+	float3 RGB05;
+	RGB05.x = ((((InputBuffer[BufferIndex].DWords[04] >> 16) << 8) | ((InputBuffer[BufferIndex].DWords[04] >> 24) & 0x0FF)) & 0x3FF) * Mul;
+	RGB05.y = ((((InputBuffer[BufferIndex].DWords[05] >> 00) << 2) | ((InputBuffer[BufferIndex].DWords[05] >> 14) & 0x003)) & 0x3FF) * Mul;
+	RGB05.z = ((((InputBuffer[BufferIndex].DWords[05] >>  8) << 4) | ((InputBuffer[BufferIndex].DWords[05] >> 20) & 0x00F)) & 0x3FF) * Mul;
+
+	//OutRGB10Buffer[Index].DWords[6] = ((RGB07.x & 0x00F) << 28) | ((RGB07.y >> 6)   << 24)   | ((RGB06.z & 0x3)  << 22) | ((RGB07.x >> 4)   << 16)							  | ((RGB06.z >> 2)   << 8)								| (RGB06.y & 0xFF);
+	float3 RGB06 = uint3(0,0,0);
+	RGB06.x = ((((InputBuffer[BufferIndex].DWords[05] >> 16) << 6) | ((InputBuffer[BufferIndex].DWords[05] >> 26) & 0x03F)) & 0x3FF) * Mul;
+	RGB06.y = ((((InputBuffer[BufferIndex].DWords[05] >> 24) << 8) | ((InputBuffer[BufferIndex].DWords[06] >> 00) & 0x0FF)) & 0x3FF) * Mul;
+	RGB06.z = ((((InputBuffer[BufferIndex].DWords[06] >>  8) << 2) | ((InputBuffer[BufferIndex].DWords[06] >> 22) & 0x003)) & 0x3FF) * Mul;
+
+	//OutRGB10Buffer[Index].DWords[7] = ((RGB08.x & 0x003) << 30) | ((RGB08.y >> 4)   << 24)								| ((RGB08.x >> 2)   << 16)							  | ((RGB07.z & 0xFF) << 8)  | ((RGB07.y & 0x3F) << 2)  | (RGB07.z >> 8);
+	float3 RGB07 = uint3(0,0,0);
+	RGB07.x = ((((InputBuffer[BufferIndex].DWords[06] >> 16) << 4) | ((InputBuffer[BufferIndex].DWords[06] >> 28) & 0x00F)) & 0x3FF) * Mul;
+	RGB07.y = ((((InputBuffer[BufferIndex].DWords[06] >> 24) << 6) | ((InputBuffer[BufferIndex].DWords[07] >> 02) & 0x03F)) & 0x3FF) * Mul;
+	RGB07.z = ((((InputBuffer[BufferIndex].DWords[07] >>  0) << 8) | ((InputBuffer[BufferIndex].DWords[07] >>  8) & 0x0FF)) & 0x3FF) * Mul;
+
+	//OutRGB10Buffer[Index].DWords[8] =							  ((RGB09.y >> 2)   << 24)								| ((RGB09.x & 0xFF) << 16) | ((RGB08.z & 0x3F) << 10) | ((RGB09.x >> 8)   << 8)  | ((RGB08.y & 0xF)  << 4)  | (RGB08.z >> 6);
+	float3 RGB08 = uint3(0,0,0);
+	RGB08.x = ((((InputBuffer[BufferIndex].DWords[07] >> 16) << 2) | ((InputBuffer[BufferIndex].DWords[07] >> 30) & 0x003)) & 0x3FF) * Mul;
+	RGB08.y = ((((InputBuffer[BufferIndex].DWords[07] >> 24) << 4) | ((InputBuffer[BufferIndex].DWords[ 8] >> 04) & 0x00F)) & 0x3FF) * Mul;
+	RGB08.z = ((((InputBuffer[BufferIndex].DWords[ 8] >>  0) << 6) | ((InputBuffer[BufferIndex].DWords[ 8] >> 10) & 0x03F)) & 0x3FF) * Mul;
+
+	//OutRGB10Buffer[Index].DWords[9] =							  ((RGB10.y & 0xFF) << 24)   | ((RGB10.x & 0x3F) << 18) | ((RGB10.y >> 8)   << 16) | ((RGB09.z & 0xF)  << 12) | ((RGB10.x >> 6)   << 8)  | ((RGB09.y & 0x3)  << 6)  | (RGB09.z >> 4);
+	float3 RGB09 = uint3(0,0,0);
+	RGB09.x = ((((InputBuffer[BufferIndex].DWords[ 8] >>  8) << 8) | ((InputBuffer[BufferIndex].DWords[ 8] >> 16) & 0x0FF)) & 0x3FF) * Mul;
+	RGB09.y = ((((InputBuffer[BufferIndex].DWords[ 8] >> 24) << 2) | ((InputBuffer[BufferIndex].DWords[ 9] >> 06) & 0x003)) & 0x3FF) * Mul;
+	RGB09.z = ((((InputBuffer[BufferIndex].DWords[ 9] >>  0) << 4) | ((InputBuffer[BufferIndex].DWords[ 9] >> 12) & 0x00F)) & 0x3FF) * Mul;
+
+	//OutRGB10Buffer[Index].DWords[10] = ((RGB11.y & 0x03F) << 26) | ((RGB11.z >> 8)   << 24)  | ((RGB11.x & 0xF)  << 20) | ((RGB11.y >> 6)   << 16) | ((RGB10.z & 0x3)  << 14) | ((RGB11.x >> 4)   << 8)								| (RGB10.z >> 2);
+	float3 RGB10 = uint3(0,0,0);
+	RGB10.x = ((((InputBuffer[BufferIndex].DWords[ 9] >>  8) << 6) | ((InputBuffer[BufferIndex].DWords[ 9] >> 18) & 0x03F)) & 0x3FF) * Mul;
+	RGB10.y = ((((InputBuffer[BufferIndex].DWords[ 9] >> 16) << 8) | ((InputBuffer[BufferIndex].DWords[ 9] >> 24) & 0x0FF)) & 0x3FF) * Mul;
+	RGB10.z = ((((InputBuffer[BufferIndex].DWords[10] >>  0) << 2) | ((InputBuffer[BufferIndex].DWords[10] >> 14) & 0x003)) & 0x3FF) * Mul;
+
+	//OutRGB10Buffer[Index].DWords[11] = ((RGB12.y & 0x00F) << 28) | ((RGB12.z >> 6)   << 24)  | ((RGB12.x & 0x3)  << 22) | ((RGB12.y >> 4)   << 16)							  | ((RGB12.x >> 2)   << 8)								| (RGB11.z & 0xFF);
+	float3 RGB11 = uint3(0,0,0);
+	RGB11.x = ((((InputBuffer[BufferIndex].DWords[10] >>  8) << 4) | ((InputBuffer[BufferIndex].DWords[10] >> 20) & 0x00F)) & 0x3FF) * Mul;
+	RGB11.y = ((((InputBuffer[BufferIndex].DWords[10] >> 16) << 6) | ((InputBuffer[BufferIndex].DWords[10] >> 26) & 0x03F)) & 0x3FF) * Mul;
+	RGB11.z = ((((InputBuffer[BufferIndex].DWords[10] >>  24) << 8) | ((InputBuffer[BufferIndex].DWords[11] >> 00) & 0x0FF)) & 0x3FF) * Mul;
+	
+	//OutRGB10Buffer[Index].DWords[12] = ((RGB13.y & 0x003) << 30) | ((RGB13.z >> 4)   << 24)								| ((RGB13.y >> 2)   << 16)							  | ((RGB13.x & 0xFF) << 8)  | ((RGB12.z & 0x3F) << 2)  | (RGB13.x >> 8);
+	float3 RGB12 = uint3(0,0,0);
+	RGB12.x = ((((InputBuffer[BufferIndex].DWords[11] >>  8) << 2) | ((InputBuffer[BufferIndex].DWords[11] >> 22) & 0x003)) & 0x3FF) * Mul;
+	RGB12.y = ((((InputBuffer[BufferIndex].DWords[11] >> 16) << 4) | ((InputBuffer[BufferIndex].DWords[11] >> 28) & 0x00F)) & 0x3FF) * Mul;
+	RGB12.z = ((((InputBuffer[BufferIndex].DWords[11] >> 24) << 6) | ((InputBuffer[BufferIndex].DWords[12] >> 02) & 0x03F)) & 0x3FF) * Mul;
+
+	//OutRGB10Buffer[Index].DWords[13] =							   ((RGB14.z >> 2)   << 24)								| ((RGB14.y & 0xFF) << 16) | ((RGB14.x & 0x3F) << 10) | ((RGB14.y >> 8)   << 8)  | ((RGB13.z & 0xF)  << 4)  | (RGB14.x >> 6);
+	float3 RGB13 = uint3(0,0,0);
+	RGB13.x = ((((InputBuffer[BufferIndex].DWords[12] >>  0) << 8) | ((InputBuffer[BufferIndex].DWords[12] >>  8) & 0x0FF)) & 0x3FF) * Mul;
+	RGB13.y = ((((InputBuffer[BufferIndex].DWords[12] >> 16) << 2) | ((InputBuffer[BufferIndex].DWords[12] >> 30) & 0x003)) & 0x3FF) * Mul;
+	RGB13.z = ((((InputBuffer[BufferIndex].DWords[12] >> 24) << 4) | ((InputBuffer[BufferIndex].DWords[13] >> 04) & 0x00F)) & 0x3FF) * Mul;
+
+	//OutRGB10Buffer[Index].DWords[14] =							   ((RGB15.z & 0xFF) << 24)  | ((RGB15.y & 0x3F) << 18) | ((RGB15.z >> 8)   << 16) | ((RGB15.x & 0xF)  << 12) | ((RGB15.y >> 6)   << 8)  | ((RGB14.z & 0x3)  << 6)  | (RGB15.x >> 4);
+	float3 RGB14 = uint3(0,0,0);
+	RGB14.x = ((((InputBuffer[BufferIndex].DWords[13] >>  0) << 6) | ((InputBuffer[BufferIndex].DWords[13] >> 10) & 0x03F)) & 0x3FF) * Mul;
+	RGB14.y = ((((InputBuffer[BufferIndex].DWords[13] >>  8) << 8) | ((InputBuffer[BufferIndex].DWords[13] >> 16) & 0x0FF)) & 0x3FF) * Mul;
+	RGB14.z = ((((InputBuffer[BufferIndex].DWords[13] >> 24) << 2) | ((InputBuffer[BufferIndex].DWords[14] >> 06) & 0x003)) & 0x3FF) * Mul;
+
+	float3 RGB15 = uint3(0,0,0);
+	RGB15.x = ((((InputBuffer[BufferIndex].DWords[14] >>  0) << 4) | ((InputBuffer[BufferIndex].DWords[14] >> 12) & 0x00F)) & 0x3FF) * Mul;
+	RGB15.y = ((((InputBuffer[BufferIndex].DWords[14] >>  8) << 6) | ((InputBuffer[BufferIndex].DWords[14] >> 18) & 0x03F)) & 0x3FF) * Mul;
+	RGB15.z = ((((InputBuffer[BufferIndex].DWords[14] >> 16) << 8) | ((InputBuffer[BufferIndex].DWords[14] >> 24) & 0x0FF)) & 0x3FF) * Mul;
+
+	#if DO_SRGB_TO_LINEAR
+		RGB00 = sRGBToLinear(RGB00);
+		RGB01 = sRGBToLinear(RGB01);
+		RGB02 = sRGBToLinear(RGB02);
+		RGB03 = sRGBToLinear(RGB03);
+		RGB04 = sRGBToLinear(RGB04);
+		RGB05 = sRGBToLinear(RGB05);
+		RGB06 = sRGBToLinear(RGB06);
+		RGB07 = sRGBToLinear(RGB07);
+		RGB08 = sRGBToLinear(RGB08);
+		RGB09 = sRGBToLinear(RGB09);
+		RGB10 = sRGBToLinear(RGB10);
+		RGB11 = sRGBToLinear(RGB11);
+		RGB12 = sRGBToLinear(RGB12);
+		RGB13 = sRGBToLinear(RGB13);
+		RGB14 = sRGBToLinear(RGB14);
+		RGB15 = sRGBToLinear(RGB15);
+	#endif
+
+	OutTexture[uint2(BaseU + 0, BaseV)] = float4(RGB00, 1.0);
+	OutTexture[uint2(BaseU + 1, BaseV)] = float4(RGB01, 1.0);
+	OutTexture[uint2(BaseU + 2, BaseV)] = float4(RGB02, 1.0);
+	OutTexture[uint2(BaseU + 3, BaseV)] = float4(RGB03, 1.0);
+	OutTexture[uint2(BaseU + 4, BaseV)] = float4(RGB04, 1.0);
+	OutTexture[uint2(BaseU + 5, BaseV)] = float4(RGB05, 1.0);
+	OutTexture[uint2(BaseU + 6, BaseV)] = float4(RGB06, 1.0);
+	OutTexture[uint2(BaseU + 7, BaseV)] = float4(RGB07, 1.0);
+	OutTexture[uint2(BaseU + 8, BaseV)] = float4(RGB08, 1.0);
+	OutTexture[uint2(BaseU + 9, BaseV)] = float4(RGB09, 1.0);
+	OutTexture[uint2(BaseU + 10, BaseV)] = float4(RGB10, 1.0);
+	OutTexture[uint2(BaseU + 11, BaseV)] = float4(RGB11, 1.0);
+	OutTexture[uint2(BaseU + 12, BaseV)] = float4(RGB12, 1.0);
+	OutTexture[uint2(BaseU + 13, BaseV)] = float4(RGB13, 1.0);
+	OutTexture[uint2(BaseU + 14, BaseV)] = float4(RGB14, 1.0);
+	OutTexture[uint2(BaseU + 15, BaseV)] = float4(RGB15, 1.0);
+}
diff --git a/Engine/Plugins/VirtualProduction/Rivermax/RivermaxCore/Source/RivermaxCore/Private/RivermaxInputStream.cpp b/Engine/Plugins/VirtualProduction/Rivermax/RivermaxCore/Source/RivermaxCore/Private/RivermaxInputStream.cpp
--- a/Engine/Plugins/VirtualProduction/Rivermax/RivermaxCore/Source/RivermaxCore/Private/RivermaxInputStream.cpp
+++ b/Engine/Plugins/VirtualProduction/Rivermax/RivermaxCore/Source/RivermaxCore/Private/RivermaxInputStream.cpp
@@ -117,9 +117,9 @@ namespace UE::RivermaxCore::Private
 			const rmax_in_stream_type StreamType = RMAX_RAW_PACKET;
 			sockaddr_in RivermaxInterface;
 			memset(&RivermaxInterface, 0, sizeof(RivermaxInterface));
-			if (inet_pton(AF_INET, StringCast<ANSICHAR>(*Options.SourceAddress).Get(), &RivermaxInterface.sin_addr) != 1)
+			if (inet_pton(AF_INET, StringCast<ANSICHAR>(*Options.InterfaceAddress).Get(), &RivermaxInterface.sin_addr) != 1)
 			{
-				UE_LOG(LogRivermax, Warning, TEXT("inet_pton failed to %s"), *Options.SourceAddress);
+				UE_LOG(LogRivermax, Warning, TEXT("inet_pton failed to %s"), *Options.InterfaceAddress);
 			}
 			else
 			{
@@ -129,9 +129,9 @@ namespace UE::RivermaxCore::Private
 				memset(&FlowAttribute, 0, sizeof(FlowAttribute));
 				FlowAttribute.local_addr.sin_family = AF_INET;
 				FlowAttribute.flow_id = FlowId;
-				if (inet_pton(AF_INET, StringCast<ANSICHAR>(*Options.DestinationAddress).Get(), &FlowAttribute.local_addr.sin_addr) != 1)
+				if (inet_pton(AF_INET, StringCast<ANSICHAR>(*Options.StreamAddress).Get(), &FlowAttribute.local_addr.sin_addr) != 1)
 				{
-					UE_LOG(LogRivermax, Warning, TEXT("inet_pton failed to %s"), *Options.DestinationAddress);
+					UE_LOG(LogRivermax, Warning, TEXT("inet_pton failed to %s"), *Options.StreamAddress);
 				}
 				else
 				{
@@ -379,7 +379,7 @@ namespace UE::RivermaxCore::Private
 									FRivermaxInputVideoFrameDescriptor Descriptor;
 									Descriptor.Width = Options.Resolution.X;
 									Descriptor.Height = Options.Resolution.Y;
-									Descriptor.Stride = Options.Resolution.X * 2;//todo stride. Bytes per row
+									Descriptor.Stride = Options.Stride;
 									FRivermaxInputVideoFrameReception NewFrame;
 									NewFrame.VideoBuffer = StreamData.CurrentFrame;
 									Listener->OnVideoFrameReceived(Descriptor, NewFrame);//todo stride
@@ -426,9 +426,7 @@ namespace UE::RivermaxCore::Private
 
 		FRivermaxInputVideoFrameDescriptor Descriptor;
 		FRivermaxInputVideoFrameRequest Request;
-		const uint32 Groups = Options.Resolution.X / PixelsPerGroup_422_8b;
-		const uint32 BytesPerLine = Groups * BytesPerGroup_422_8b;
-		Descriptor.VideoBufferSize = Options.Resolution.Y * BytesPerLine;
+		Descriptor.VideoBufferSize = Options.Resolution.Y * Options.Stride;
 		Listener->OnVideoFrameRequested(Descriptor, Request);
 		StreamData.CurrentFrame = Request.VideoBuffer;
 		StreamData.WritingOffset = 0;
diff --git a/Engine/Plugins/VirtualProduction/Rivermax/RivermaxCore/Source/RivermaxCore/Private/RivermaxUtils.cpp b/Engine/Plugins/VirtualProduction/Rivermax/RivermaxCore/Source/RivermaxCore/Private/RivermaxUtils.cpp
--- a/Engine/Plugins/VirtualProduction/Rivermax/RivermaxCore/Source/RivermaxCore/Private/RivermaxUtils.cpp
+++ b/Engine/Plugins/VirtualProduction/Rivermax/RivermaxCore/Source/RivermaxCore/Private/RivermaxUtils.cpp
@@ -68,8 +68,8 @@ namespace UE::RivermaxCore::Private::Utils
 		OutSDPDescription.Appendf("v=0\n");
 		OutSDPDescription.Appendf("s=SMPTE ST2110 20 streams\n");
 		OutSDPDescription.Appendf("m=video %d RTP/AVP 96\n", Options.Port);
-		OutSDPDescription.Appendf("c=IN IP4 %S/64\n", *Options.DestinationAddress);
-		OutSDPDescription.Appendf("a=source-filter: incl IN IP4 %S %S\n", *Options.DestinationAddress, *Options.SourceAddress);
+		OutSDPDescription.Appendf("c=IN IP4 %S/64\n", *Options.StreamAddress);
+		OutSDPDescription.Appendf("a=source-filter: incl IN IP4 %S %S\n", *Options.StreamAddress, *Options.InterfaceAddress);
 		OutSDPDescription.Appendf("a=rtpmap:96 raw/90000\n");
 		OutSDPDescription.Appendf("a=fmtp: 96 sampling=%S; width=%d; height=%d; exactframerate=%S; depth=%S; TCS=SDR; colorimetry=BT709; PM=2110GPM; SSN=ST2110-20:2017; TP=2110TPN;\n"
 			, *PixelFormatToSamplingDesc(Options.PixelFormat)
diff --git a/Engine/Plugins/VirtualProduction/Rivermax/RivermaxCore/Source/RivermaxCore/Private/RivermaxUtils.h b/Engine/Plugins/VirtualProduction/Rivermax/RivermaxCore/Source/RivermaxCore/Private/RivermaxUtils.h
--- a/Engine/Plugins/VirtualProduction/Rivermax/RivermaxCore/Source/RivermaxCore/Private/RivermaxUtils.h
+++ b/Engine/Plugins/VirtualProduction/Rivermax/RivermaxCore/Source/RivermaxCore/Private/RivermaxUtils.h
@@ -14,8 +14,6 @@ namespace UE::RivermaxCore::Private::Utils
 	/** Various constants used for stream initialization */
 	static constexpr uint32 BytesPerHeader = 20;
 	static constexpr uint32 BytesPerPacket = 1200;
-	static constexpr uint32 PixelsPerGroup_422_8b = 2;
-	static constexpr uint32 BytesPerGroup_422_8b = 4;
 	static constexpr uint32 FullHDHeight = 1080;
 	static constexpr uint32 FullHDWidth = 1920;
 	static constexpr uint32 MaxPayloadSize = 1280;
diff --git a/Engine/Plugins/VirtualProduction/Rivermax/RivermaxCore/Source/RivermaxCore/Public/RivermaxTypes.h b/Engine/Plugins/VirtualProduction/Rivermax/RivermaxCore/Source/RivermaxCore/Public/RivermaxTypes.h
--- a/Engine/Plugins/VirtualProduction/Rivermax/RivermaxCore/Source/RivermaxCore/Public/RivermaxTypes.h
+++ b/Engine/Plugins/VirtualProduction/Rivermax/RivermaxCore/Source/RivermaxCore/Public/RivermaxTypes.h
@@ -25,11 +25,11 @@ namespace UE::RivermaxCore
 		/** Stream FrameRate */
 		FFrameRate FrameRate = { 24,1 };
 
-		/** Source IP to bind to */
-		FString SourceAddress;
+		/** Interface IP to bind to */
+		FString InterfaceAddress;
 
-		/** Destination IP to send to. Defaults to multicast group IP. */
-		FString DestinationAddress = TEXT("224.1.1.1");
+		/** IP of the stream. Defaults to multicast group IP. */
+		FString StreamAddress = TEXT("224.1.1.1");
 
 		/** Port to be used by stream */
 		uint32 Port = 50000;
diff --git a/Engine/Plugins/VirtualProduction/Rivermax/RivermaxCore/Source/RivermaxRendering/Private/RivermaxShaders.cpp b/Engine/Plugins/VirtualProduction/Rivermax/RivermaxCore/Source/RivermaxRendering/Private/RivermaxShaders.cpp
--- a/Engine/Plugins/VirtualProduction/Rivermax/RivermaxCore/Source/RivermaxRendering/Private/RivermaxShaders.cpp
+++ b/Engine/Plugins/VirtualProduction/Rivermax/RivermaxCore/Source/RivermaxRendering/Private/RivermaxShaders.cpp
@@ -41,6 +41,22 @@ namespace UE::RivermaxShaders
 		return Parameters;
 	}
 
+	/* FRGB8BitToRGBA8CS shader
+	 *****************************************************************************/
+
+	IMPLEMENT_GLOBAL_SHADER(FRGB8BitToRGBA8CS, "/Plugin/RivermaxCore/Private/RivermaxShaders.usf", "RGB8BitToRGBA8CS", SF_Compute);
+
+	FRGB8BitToRGBA8CS::FParameters* FRGB8BitToRGBA8CS::AllocateAndSetParameters(FRDGBuilder& GraphBuilder, FRDGBufferRef RGBBuffer, FRDGTextureRef OutputTexture, int32 BufferElementsPerRow)
+	{
+		FRGB8BitToRGBA8CS::FParameters* Parameters = GraphBuilder.AllocParameters<FRGB8BitToRGBA8CS::FParameters>();
+
+		Parameters->InputRGB8Buffer = GraphBuilder.CreateSRV(RGBBuffer);
+		Parameters->HorizontalElementCount = BufferElementsPerRow;
+		Parameters->OutTexture = GraphBuilder.CreateUAV(OutputTexture, ERDGUnorderedAccessViewFlags::None);
+
+		return Parameters;
+	}
+
 	/* FRGBToRGB8BitCS shader
 	 *****************************************************************************/
 
@@ -71,6 +87,23 @@ namespace UE::RivermaxShaders
 		return Parameters;
 	}
 
+	/* FRGB10BitToRGBA10CS shader
+	 *****************************************************************************/
+
+	IMPLEMENT_GLOBAL_SHADER(FRGB10BitToRGBA10CS, "/Plugin/RivermaxCore/Private/RivermaxShaders.usf", "RGB10BitToRGBACS", SF_Compute);
+
+
+	FRGB10BitToRGBA10CS::FParameters* FRGB10BitToRGBA10CS::AllocateAndSetParameters(FRDGBuilder& GraphBuilder, FRDGBufferRef RGBBuffer, FRDGTextureRef OutputTexture, int32 BufferElementsPerRow)
+	{
+		FRGB10BitToRGBA10CS::FParameters* Parameters = GraphBuilder.AllocParameters<FRGB10BitToRGBA10CS::FParameters>();
+
+		Parameters->InputBuffer = GraphBuilder.CreateSRV(RGBBuffer);
+		Parameters->HorizontalElementCount = BufferElementsPerRow;
+		Parameters->OutTexture = GraphBuilder.CreateUAV(OutputTexture, ERDGUnorderedAccessViewFlags::None);
+
+		return Parameters;
+	}
+
 	/* FRGBToRGB10BitCS shader
 	 *****************************************************************************/
 
diff --git a/Engine/Plugins/VirtualProduction/Rivermax/RivermaxCore/Source/RivermaxRendering/Public/RivermaxShaders.h b/Engine/Plugins/VirtualProduction/Rivermax/RivermaxCore/Source/RivermaxRendering/Public/RivermaxShaders.h
--- a/Engine/Plugins/VirtualProduction/Rivermax/RivermaxCore/Source/RivermaxRendering/Public/RivermaxShaders.h
+++ b/Engine/Plugins/VirtualProduction/Rivermax/RivermaxCore/Source/RivermaxRendering/Public/RivermaxShaders.h
@@ -57,6 +57,46 @@ namespace UE::RivermaxShaders
 		FRGBToYUV10Bit422LittleEndianCS::FParameters* AllocateAndSetParameters(FRDGBuilder& GraphBuilder, FRDGTextureRef RGBATexture, const FIntPoint& SourceSize, const FIntRect& SourceViewRect, const FIntPoint& OutputSize, const FMatrix& ColorTransform, const FVector& YUVOffset, bool bDoLinearToSrgb, FRDGBufferRef OutputBuffer);
 	};
 
+	/**
+	 * Compute shader to convert packed 8 bits RGB to RGBA 8bits
+	 */
+	class RIVERMAXRENDERING_API FRGB8BitToRGBA8CS : public FGlobalShader
+	{
+	public:
+
+		// Structure definition to match StructuredBuffer in RivermaxShaders.usf
+		// For 8bit output, it's 3 bytes per pixels. To align with 4 bytes (32bit) we use 12 bytes (4 pixels) per output
+		struct FRGB8BitBuffer
+		{
+			uint32 DWord0;
+			uint32 DWord1;
+			uint32 DWord2;
+		};
+
+		DECLARE_GLOBAL_SHADER(FRGB8BitToRGBA8CS);
+		SHADER_USE_PARAMETER_STRUCT(FRGB8BitToRGBA8CS, FGlobalShader);
+
+		class FSRGBToLinear : SHADER_PERMUTATION_BOOL("DO_SRGB_TO_LINEAR");
+		using FPermutationDomain = TShaderPermutationDomain<FSRGBToLinear>;
+
+		BEGIN_SHADER_PARAMETER_STRUCT(FParameters, )
+			SHADER_PARAMETER_RDG_BUFFER_SRV(StructuredBuffer<FRGB8BitBuffer>, InputRGB8Buffer)
+			SHADER_PARAMETER(uint32, HorizontalElementCount)
+			SHADER_PARAMETER_RDG_TEXTURE_UAV(RWTexture2D, OutTexture)
+		END_SHADER_PARAMETER_STRUCT()
+
+	public:
+
+		// Called by the engine to determine which permutations to compile for this shader
+		static bool ShouldCompilePermutation(const FGlobalShaderPermutationParameters& Parameters)
+		{
+			return IsFeatureLevelSupported(Parameters.Platform, ERHIFeatureLevel::SM5);
+		}
+
+		/** Allocates and setup shader parameter in the incoming graph builder */
+		FRGB8BitToRGBA8CS::FParameters* AllocateAndSetParameters(FRDGBuilder& GraphBuilder, FRDGBufferRef RGBBuffer, FRDGTextureRef OutputTexture, int32 BufferElementsPerRow);
+	};
+
 	/**
 	 * Compute shader to convert RGBA to packed 8 bits RGB
 	 */
@@ -130,7 +170,7 @@ namespace UE::RivermaxShaders
 			SHADER_PARAMETER(uint32, HorizontalElementCount)
 			SHADER_PARAMETER(uint32, InputTexturePixelsPerThread)
 			SHADER_PARAMETER_RDG_BUFFER_UAV(RWStructuredBuffer<FRGB10BitBuffer>, OutRGB10Buffer)
-		END_SHADER_PARAMETER_STRUCT()
+			END_SHADER_PARAMETER_STRUCT()
 
 	public:
 
@@ -144,6 +184,45 @@ namespace UE::RivermaxShaders
 		FRGBToRGB10BitCS::FParameters* AllocateAndSetParameters(FRDGBuilder& GraphBuilder, FRDGTextureRef RGBATexture, const FIntPoint& SourceSize, const FIntRect& SourceViewRect, const FIntPoint& OutputSize, FRDGBufferRef OutputBuffer);
 	};
 
+	/**
+	 * Compute shader to convert packed 10 bits RGB to RGBA
+	 */
+	class RIVERMAXRENDERING_API FRGB10BitToRGBA10CS : public FGlobalShader
+	{
+	public:
+
+		// Structure definition to match StructuredBuffer in RivermaxShaders.usf
+		// 10 bits per channel, 30 bits per pixel. Aligns on 15 32bits -> 480bits -> 16RGB10 pixels
+		// 2110-20 says it's 4 pixels per pgroup. 
+		struct FRGB10BitBuffer
+		{
+			uint32 DWords[15];
+		};
+
+		DECLARE_GLOBAL_SHADER(FRGB10BitToRGBA10CS);
+		SHADER_USE_PARAMETER_STRUCT(FRGB10BitToRGBA10CS, FGlobalShader);
+		
+		class FSRGBToLinear : SHADER_PERMUTATION_BOOL("DO_SRGB_TO_LINEAR");
+		using FPermutationDomain = TShaderPermutationDomain<FSRGBToLinear>;
+
+		BEGIN_SHADER_PARAMETER_STRUCT(FParameters, )
+			SHADER_PARAMETER_RDG_BUFFER_SRV(StructuredBuffer<FRGB10BitBuffer>, InputBuffer)
+			SHADER_PARAMETER(uint32, HorizontalElementCount)
+			SHADER_PARAMETER_RDG_TEXTURE_UAV(RWTexture2D, OutTexture)
+		END_SHADER_PARAMETER_STRUCT()
+
+	public:
+
+		// Called by the engine to determine which permutations to compile for this shader
+		static bool ShouldCompilePermutation(const FGlobalShaderPermutationParameters& Parameters)
+		{
+			return IsFeatureLevelSupported(Parameters.Platform, ERHIFeatureLevel::SM5);
+		}
+
+		/** Allocates and setup shader parameter in the incoming graph builder */
+		FRGB10BitToRGBA10CS::FParameters* AllocateAndSetParameters(FRDGBuilder& GraphBuilder, FRDGBufferRef RGBBuffer, FRDGTextureRef OutputTexture, int32 BufferElementsPerRow);
+	};
+
 	/**
 	 * Compute shader to convert RGBA to Float16 RGB
 	 */
diff --git a/Engine/Plugins/VirtualProduction/Rivermax/RivermaxMedia/Source/RivermaxMedia/Private/RivermaxMediaCapture.cpp b/Engine/Plugins/VirtualProduction/Rivermax/RivermaxMedia/Source/RivermaxMedia/Private/RivermaxMediaCapture.cpp
--- a/Engine/Plugins/VirtualProduction/Rivermax/RivermaxMedia/Source/RivermaxMedia/Private/RivermaxMediaCapture.cpp
+++ b/Engine/Plugins/VirtualProduction/Rivermax/RivermaxMedia/Source/RivermaxMedia/Private/RivermaxMediaCapture.cpp
@@ -202,8 +202,8 @@ bool URivermaxMediaCapture::ConfigureStream(URivermaxMediaOutput* InMediaOutput,
 {
 	using namespace UE::RivermaxCore;
 
-	OutOptions.DestinationAddress = InMediaOutput->DestinationAddress;
-	OutOptions.SourceAddress = InMediaOutput->SourceAddress;
+	OutOptions.StreamAddress = InMediaOutput->StreamAddress;
+	OutOptions.InterfaceAddress = InMediaOutput->InterfaceAddress;
 	OutOptions.Port = InMediaOutput->Port;
 	OutOptions.Resolution = InMediaOutput->Resolution;
 	OutOptions.FrameRate = InMediaOutput->FrameRate;
@@ -263,7 +263,10 @@ void URivermaxMediaCapture::OnFrameCaptured_RenderingThread(const FCaptureBaseDa
 	NewFrame.Stride = BytesPerRow;
 	NewFrame.VideoBuffer = InBuffer;
 	NewFrame.FrameIdentifier = InBaseData.SourceFrameNumber;
-	RivermaxStream->PushVideoFrame(NewFrame);
+	if (RivermaxStream->PushVideoFrame(NewFrame) == false)
+	{
+		UE_LOG(LogRivermaxMedia, Warning, TEXT("Failed to pushed captured frame"));
+	}
 }
 
 void URivermaxMediaCapture::OnRHIResourceCaptured_RenderingThread(const FCaptureBaseData& InBaseData,	TSharedPtr<FMediaCaptureUserData, ESPMode::ThreadSafe> InUserData, FTextureRHIRef InTexture)
diff --git a/Engine/Plugins/VirtualProduction/Rivermax/RivermaxMedia/Source/RivermaxMedia/Private/RivermaxMediaPlayer.cpp b/Engine/Plugins/VirtualProduction/Rivermax/RivermaxMedia/Source/RivermaxMedia/Private/RivermaxMediaPlayer.cpp
--- a/Engine/Plugins/VirtualProduction/Rivermax/RivermaxMedia/Source/RivermaxMedia/Private/RivermaxMediaPlayer.cpp
+++ b/Engine/Plugins/VirtualProduction/Rivermax/RivermaxMedia/Source/RivermaxMedia/Private/RivermaxMediaPlayer.cpp
@@ -35,7 +35,7 @@ namespace UE::RivermaxMedia
 		, MaxNumVideoFrameBuffer(8)
 		, RivermaxThreadNewState(EMediaState::Closed)
 		, EventSink(InEventSink)
-		, bIsSRGBInput(true)
+		, bIsSRGBInput(false)
 		, bUseVideo(false)
 		, SupportedSampleTypes(EMediaIOSampleType::None)
 		, bPauseRequested(false)
@@ -63,18 +63,12 @@ namespace UE::RivermaxMedia
 			return false;
 		}
 
+		//Video related options
 		{
-			StreamOptions.SourceAddress = Options->GetMediaOption(RivermaxMediaOption::SourceAddress, FString());
-			StreamOptions.DestinationAddress = Options->GetMediaOption(RivermaxMediaOption::DestinationAddress, FString());
-			StreamOptions.Port = Options->GetMediaOption(RivermaxMediaOption::Port, (int64)0);
-
-			ERivermaxMediaOutputPixelFormat DesiredFormat = ERivermaxMediaOutputPixelFormat::PF_8BIT_YUV;
-			DesiredFormat = (ERivermaxMediaOutputPixelFormat)Options->GetMediaOption(RivermaxMediaOption::PixelFormat, (int64)DesiredFormat);
-			StreamOptions.PixelFormat = DesiredFormat == ERivermaxMediaOutputPixelFormat::PF_8BIT_YUV ? ERivermaxOutputPixelFormat::RMAX_8BIT_YCBCR : ERivermaxOutputPixelFormat::RMAX_8BIT_YCBCR;
-			StreamOptions.Resolution = VideoTrackFormat.Dim;
-			StreamOptions.FrameRate = VideoFrameRate;
-
+			bIsSRGBInput = Options->GetMediaOption(RivermaxMediaOption::SRGBInput, bIsSRGBInput);
+			DesiredPixelFormat = (ERivermaxMediaSourePixelFormat)Options->GetMediaOption(RivermaxMediaOption::PixelFormat, (int64)ERivermaxMediaSourePixelFormat::RGB_8bit);
 		}
+
 		{
 			//Adjust supported sample types based on what's being captured
 			SupportedSampleTypes = EMediaIOSampleType::Video;
@@ -85,7 +79,7 @@ namespace UE::RivermaxMedia
 		bUseVideo = true;
 
 		IRivermaxCoreModule* Module = FModuleManager::GetModulePtr<IRivermaxCoreModule>("RivermaxCore");
-		if (Module)
+		if (Module && ConfigureStream(Options))
 		{
 			InputStream = Module->CreateInputStream();
 		}
@@ -187,10 +181,9 @@ namespace UE::RivermaxMedia
 
 		if (bUseVideo && ReceivedVideoFrame.VideoBuffer)
 		{
-			const EMediaTextureSampleFormat VideoSampleFormat = EMediaTextureSampleFormat::CharUYVY;
 			if (RivermaxThreadCurrentTextureSample.IsValid())
 			{
-				if (RivermaxThreadCurrentTextureSample->SetProperties(FrameInfo.Stride, FrameInfo.Width, FrameInfo.Height, VideoSampleFormat, DecodedTime, VideoFrameRate, DecodedTimecode, bIsSRGBInput))
+				if (RivermaxThreadCurrentTextureSample->ConfigureSample(FrameInfo.Width, FrameInfo.Height, FrameInfo.Stride, DesiredPixelFormat, DecodedTime, VideoFrameRate, DecodedTimecode, bIsSRGBInput))
 				{
 					Samples->AddVideo(RivermaxThreadCurrentTextureSample.ToSharedRef());
 				}
@@ -281,6 +274,45 @@ namespace UE::RivermaxMedia
 	{
 		RivermaxThreadNewState = bHasSucceed ? EMediaState::Playing : EMediaState::Error;
 	}
+
+	bool FRivermaxMediaPlayer::ConfigureStream(const IMediaOptions* Options)
+	{
+		StreamOptions.InterfaceAddress = Options->GetMediaOption(RivermaxMediaOption::InterfaceAddress, FString());
+		StreamOptions.StreamAddress = Options->GetMediaOption(RivermaxMediaOption::StreamAddress, FString());
+		StreamOptions.Port = Options->GetMediaOption(RivermaxMediaOption::Port, (int64)0);
+		StreamOptions.Resolution = VideoTrackFormat.Dim;
+		StreamOptions.FrameRate = VideoFrameRate;
+
+		switch (DesiredPixelFormat)
+		{
+		case ERivermaxMediaSourePixelFormat::YUV422_8bit:
+		{
+			StreamOptions.PixelFormat = ERivermaxOutputPixelFormat::RMAX_8BIT_YCBCR;
+			StreamOptions.Stride = StreamOptions.Resolution.X / 2 * 4; // 4 bytes for a group of 2 pixels
+			break;
+		}
+		case ERivermaxMediaSourePixelFormat::RGB_8bit:
+		{
+			StreamOptions.PixelFormat = ERivermaxOutputPixelFormat::RMAX_8BIT_RGB;
+			StreamOptions.Stride = StreamOptions.Resolution.X * 3; //3 bytes per pixel
+			break;
+		}
+		case ERivermaxMediaSourePixelFormat::RGB_10bit:
+		{
+			StreamOptions.PixelFormat = ERivermaxOutputPixelFormat::RMAX_10BIT_RGB;
+			StreamOptions.Stride = StreamOptions.Resolution.X / 4 * 15; //15 bytes for a group of 4 pixels (4 * 30bits / 8 = 15)
+			break;
+		}
+		default:
+		{
+			UE_LOG(LogRivermaxMedia, Error, TEXT("Desired pixel format (%s) is not a valid Rivermax pixel format"), *StaticEnum<ERivermaxMediaSourePixelFormat>()->GetValueAsString(DesiredPixelFormat));
+			return false;
+		}
+		}
+
+		return true;
+	}
+
 }
 
 #undef LOCTEXT_NAMESPACE
diff --git a/Engine/Plugins/VirtualProduction/Rivermax/RivermaxMedia/Source/RivermaxMedia/Private/RivermaxMediaPlayer.h b/Engine/Plugins/VirtualProduction/Rivermax/RivermaxMedia/Source/RivermaxMedia/Private/RivermaxMediaPlayer.h
--- a/Engine/Plugins/VirtualProduction/Rivermax/RivermaxMedia/Source/RivermaxMedia/Private/RivermaxMediaPlayer.h
+++ b/Engine/Plugins/VirtualProduction/Rivermax/RivermaxMedia/Source/RivermaxMedia/Private/RivermaxMediaPlayer.h
@@ -81,6 +81,9 @@ namespace  UE::RivermaxMedia
 		virtual void SetupSampleChannels() override;
 		//~ End FMediaIOCorePlayerBase interface
 
+	private:
+		bool ConfigureStream(const IMediaOptions* Options);
+
 	private:
 
 		/** Audio, MetaData, Texture  sample object pool. */
@@ -114,6 +117,9 @@ namespace  UE::RivermaxMedia
 
 		/** Flag to indicate that pause is being requested */
 		std::atomic<bool> bPauseRequested;
+
+		/** Pixel format provided by media source */
+		ERivermaxMediaSourePixelFormat DesiredPixelFormat = ERivermaxMediaSourePixelFormat::RGB_10bit;
 	};
 }
 
diff --git a/Engine/Plugins/VirtualProduction/Rivermax/RivermaxMedia/Source/RivermaxMedia/Private/RivermaxMediaSource.cpp b/Engine/Plugins/VirtualProduction/Rivermax/RivermaxMedia/Source/RivermaxMedia/Private/RivermaxMediaSource.cpp
--- a/Engine/Plugins/VirtualProduction/Rivermax/RivermaxMedia/Source/RivermaxMedia/Private/RivermaxMediaSource.cpp
+++ b/Engine/Plugins/VirtualProduction/Rivermax/RivermaxMedia/Source/RivermaxMedia/Private/RivermaxMediaSource.cpp
@@ -12,6 +12,10 @@
 
 bool URivermaxMediaSource::GetMediaOption(const FName& Key, bool DefaultValue) const
 {
+	if (Key == RivermaxMediaOption::SRGBInput)
+	{
+		return bIsSRGBInput;
+	}
 	return Super::GetMediaOption(Key, DefaultValue);
 }
 
@@ -50,23 +54,24 @@ FString URivermaxMediaSource::GetMediaOption(const FName& Key, const FString& De
 	{
 		return FString::Printf(TEXT("FormatDescriptorTodo"));
 	}
-	else if (Key == RivermaxMediaOption::SourceAddress)
+	else if (Key == RivermaxMediaOption::InterfaceAddress)
 	{
-		return SourceAddress;
+		return InterfaceAddress;
 	}
-	else if (Key == RivermaxMediaOption::DestinationAddress)
+	else if (Key == RivermaxMediaOption::StreamAddress)
 	{
-		return DestinationAddress;
+		return StreamAddress;
 	}
 	return Super::GetMediaOption(Key, DefaultValue);
 }
 
 bool URivermaxMediaSource::HasMediaOption(const FName& Key) const
 {
-	if ((Key == RivermaxMediaOption::SourceAddress) ||
-		(Key == RivermaxMediaOption::DestinationAddress) ||
+	if ((Key == RivermaxMediaOption::InterfaceAddress) ||
+		(Key == RivermaxMediaOption::StreamAddress) ||
 		(Key == RivermaxMediaOption::Port) ||
 		(Key == RivermaxMediaOption::PixelFormat) ||
+		(Key == RivermaxMediaOption::SRGBInput) ||
 		(Key == FMediaIOCoreMediaOption::FrameRateNumerator) ||
 		(Key == FMediaIOCoreMediaOption::FrameRateDenominator) ||
 		(Key == FMediaIOCoreMediaOption::ResolutionWidth) ||
diff --git a/Engine/Plugins/VirtualProduction/Rivermax/RivermaxMedia/Source/RivermaxMedia/Private/RivermaxMediaSourceOptions.h b/Engine/Plugins/VirtualProduction/Rivermax/RivermaxMedia/Source/RivermaxMedia/Private/RivermaxMediaSourceOptions.h
--- a/Engine/Plugins/VirtualProduction/Rivermax/RivermaxMedia/Source/RivermaxMedia/Private/RivermaxMediaSourceOptions.h
+++ b/Engine/Plugins/VirtualProduction/Rivermax/RivermaxMedia/Source/RivermaxMedia/Private/RivermaxMediaSourceOptions.h
@@ -7,8 +7,9 @@
 
 namespace RivermaxMediaOption
 {
-	static const FName SourceAddress("SourceAddress");
-	static const FName DestinationAddress("DestinationAddress");
+	static const FName InterfaceAddress("InterfaceAddress");
+	static const FName StreamAddress("StreamAddress");
 	static const FName Port("Port");
 	static const FName PixelFormat("PixelFormat");
+	static const FName SRGBInput("SRGBInput");
 }
diff --git a/Engine/Plugins/VirtualProduction/Rivermax/RivermaxMedia/Source/RivermaxMedia/Private/RivermaxMediaTextureSample.cpp b/Engine/Plugins/VirtualProduction/Rivermax/RivermaxMedia/Source/RivermaxMedia/Private/RivermaxMediaTextureSample.cpp
new file mode 100644
--- /dev/null
+++ b/Engine/Plugins/VirtualProduction/Rivermax/RivermaxMedia/Source/RivermaxMedia/Private/RivermaxMediaTextureSample.cpp
@@ -0,0 +1,52 @@
+// Copyright Epic Games, Inc. All Rights Reserved.
+
+#include "RivermaxMediaTextureSample.h"
+
+#include "RivermaxMediaTextureSampleConverter.h"
+
+FRivermaxMediaTextureSample::FRivermaxMediaTextureSample()
+	: TextureConverter(MakeUnique<FRivermaxMediaTextureSampleConverter>())
+{
+	
+}
+
+const FMatrix& FRivermaxMediaTextureSample::GetYUVToRGBMatrix() const
+{
+	return MediaShaders::YuvToRgbRec709Scaled;
+}
+
+IMediaTextureSampleConverter* FRivermaxMediaTextureSample::GetMediaTextureSampleConverter()
+{
+	return TextureConverter.Get();
+}
+
+bool FRivermaxMediaTextureSample::IsOutputSrgb() const
+{
+	// We always do the sRGB to Linear conversion in shader if specified by the source
+	// If true is returned here, MediaTextureResource will create (try) a sRGB texture which only works for 8 bits
+	return false;
+}
+
+bool FRivermaxMediaTextureSample::ConfigureSample(uint32 InWidth, uint32 InHeight, uint32 InStride, ERivermaxMediaSourePixelFormat InSampleFormat, FTimespan InTime, const FFrameRate& InFrameRate, const TOptional<FTimecode>& InTimecode, bool bInIsSRGBInput)
+{
+	EMediaTextureSampleFormat VideoSampleFormat;
+	switch (InSampleFormat)
+	{
+		case ERivermaxMediaSourePixelFormat::RGB_10bit:
+		{
+			VideoSampleFormat = EMediaTextureSampleFormat::CharBGR10A2;
+			break;
+		}
+		case ERivermaxMediaSourePixelFormat::YUV422_8bit:
+		case ERivermaxMediaSourePixelFormat::RGB_8bit:
+		default:
+		{
+			VideoSampleFormat = EMediaTextureSampleFormat::CharBGRA;
+			break;
+		}
+	}
+
+	TextureConverter->Setup(InSampleFormat, AsShared(), bInIsSRGBInput);
+	return Super::SetProperties(InStride, InWidth, InHeight, VideoSampleFormat, InTime, InFrameRate, InTimecode, bInIsSRGBInput);
+}
+
diff --git a/Engine/Plugins/VirtualProduction/Rivermax/RivermaxMedia/Source/RivermaxMedia/Private/RivermaxMediaTextureSample.h b/Engine/Plugins/VirtualProduction/Rivermax/RivermaxMedia/Source/RivermaxMedia/Private/RivermaxMediaTextureSample.h
--- a/Engine/Plugins/VirtualProduction/Rivermax/RivermaxMedia/Source/RivermaxMedia/Private/RivermaxMediaTextureSample.h
+++ b/Engine/Plugins/VirtualProduction/Rivermax/RivermaxMedia/Source/RivermaxMedia/Private/RivermaxMediaTextureSample.h
@@ -3,52 +3,37 @@
 #pragma once
 
 #include "MediaIOCoreTextureSampleBase.h"
+
 #include "MediaShaders.h"
+#include "RivermaxMediaSource.h"
+#include "Templates/SharedPointer.h"
+
+class FRivermaxMediaTextureSampleConverter;
 
 /**
  * Implements a media texture sample for RivermaxMedia.
  */
-class FRivermaxMediaTextureSample
-	: public FMediaIOCoreTextureSampleBase
+class FRivermaxMediaTextureSample : public FMediaIOCoreTextureSampleBase, public TSharedFromThis<FRivermaxMediaTextureSample>
 {
 	using Super = FMediaIOCoreTextureSampleBase;
 
 public:
 
-	/**
-	 * Initialize the sample.
-	 *
-	 * @param InVideoData The video frame data.
-	 * @param InSampleFormat The sample format.
-	 * @param InTime The sample time (in the player's own clock).
-	 * @param InFrameRate The framerate of the media that produce the sample.
-	 * @param InTimecode The sample timecode if available.
-	 * @param bInIsSRGB Whether the sample is in sRGB space.
-	 */
-	bool InitializeProgressive(void* InVideoBuffer, uint32 InWidth, uint32 InHeight, uint32 InVideoBufferSize, uint32 InStride, EMediaTextureSampleFormat InSampleFormat, FTimespan InTime, const FFrameRate& InFrameRate, const TOptional<FTimecode>& InTimecode, bool bInIsSRGB)
-	{
-		return Super::Initialize(InVideoBuffer
-			, InVideoBufferSize
-			, InStride
-			, InWidth
-			, InHeight
-			, InSampleFormat
-			, InTime
-			, InFrameRate
-			, InTimecode
-			, bInIsSRGB);
-	}
-
-	/**
-	 * Get YUV to RGB conversion matrix
-	 *
-	 * @return MediaIOCore Yuv To Rgb matrix
-	 */
-	virtual const FMatrix& GetYUVToRGBMatrix() const override
-	{
-		return MediaShaders::YuvToRgbRec709Scaled;
-	}
+	FRivermaxMediaTextureSample();
+
+
+	//~ Begin IMediaTextureSample interface
+	virtual const FMatrix& GetYUVToRGBMatrix() const override;
+	virtual IMediaTextureSampleConverter* GetMediaTextureSampleConverter() override;
+	virtual bool IsOutputSrgb() const override;
+	//~ End IMediaTextureSample interface
+
+	bool ConfigureSample(uint32 InWidth, uint32 InHeight, uint32 InStride, ERivermaxMediaSourePixelFormat InSampleFormat, FTimespan InTime, const FFrameRate& InFrameRate, const TOptional<FTimecode>& InTimecode, bool bInIsSRGBInput);
+
+private:
 
+	/** Texture converted used to handle incoming 2110 formats and convert them to RGB textures the engine handles */
+	TUniquePtr<FRivermaxMediaTextureSampleConverter> TextureConverter;
 };
 
 /*
diff --git a/Engine/Plugins/VirtualProduction/Rivermax/RivermaxMedia/Source/RivermaxMedia/Private/RivermaxMediaTextureSampleConverter.cpp b/Engine/Plugins/VirtualProduction/Rivermax/RivermaxMedia/Source/RivermaxMedia/Private/RivermaxMediaTextureSampleConverter.cpp
new file mode 100644
--- /dev/null
+++ b/Engine/Plugins/VirtualProduction/Rivermax/RivermaxMedia/Source/RivermaxMedia/Private/RivermaxMediaTextureSampleConverter.cpp
@@ -0,0 +1,113 @@
+// Copyright Epic Games, Inc. All Rights Reserved.
+
+#include "RivermaxMediaTextureSampleConverter.h"
+
+#include "RenderGraphBuilder.h"
+#include "RivermaxShaders.h"
+#include "RivermaxMediaTextureSample.h"
+
+
+
+void FRivermaxMediaTextureSampleConverter::Setup(ERivermaxMediaSourePixelFormat InPixelFormat, TWeakPtr<FRivermaxMediaTextureSample> InSample, bool bInDoSRGBToLinear)
+{
+	InputPixelFormat = InPixelFormat;
+	Sample = InSample;
+	bDoSRGBToLinear = bInDoSRGBToLinear;
+}
+
+bool FRivermaxMediaTextureSampleConverter::Convert(FTexture2DRHIRef& InDestinationTexture, const FConversionHints& Hints)
+{
+	using namespace UE::RivermaxShaders;
+
+	TSharedPtr<FRivermaxMediaTextureSample> SamplePtr = Sample.Pin();
+	if (SamplePtr.IsValid() == false)
+	{
+		return false;
+	}
+
+	FIntVector GroupCount;
+	
+	FRDGBuilder GraphBuilder(FRHICommandListExecutor::GetImmediateCommandList());
+	FRDGTextureRef OutputResource = GraphBuilder.RegisterExternalTexture(CreateRenderTarget(InDestinationTexture, TEXT("RivermaxMediaTextureOutputResource")));
+
+	//Configure shader and add conversion pass based on desired pixel format
+	switch (InputPixelFormat)
+	{
+	case ERivermaxMediaSourePixelFormat::YUV422_8bit:
+	{
+		break;
+	}
+	case ERivermaxMediaSourePixelFormat::RGB_8bit:
+	{
+		constexpr uint32 PixelPerGroup = 1;
+		constexpr uint32 BytesPerGroup = (PixelPerGroup * 3 * 8) / 8;
+		const uint32 HorizontalByteCount = (InDestinationTexture->GetDesc().Extent.X / PixelPerGroup) * BytesPerGroup;
+		const int32 BytesPerElement = sizeof(FRGB8BitToRGBA8CS::FRGB8BitBuffer);
+		const int32 ElementsPerRow = HorizontalByteCount / BytesPerElement;
+		const int32 ElementCount = ElementsPerRow * InDestinationTexture->GetDesc().Extent.Y;
+
+		FRGB8BitToRGBA8CS::FPermutationDomain PermutationVector;
+		PermutationVector.Set<FRGB8BitToRGBA8CS::FSRGBToLinear>(bDoSRGBToLinear);
+
+		FRDGBufferRef InputRGGBuffer = CreateStructuredBuffer(GraphBuilder, TEXT("RivermaxInputBuffer"), BytesPerElement, ElementCount, SamplePtr->GetBuffer(), BytesPerElement * ElementCount);
+		constexpr int32 PixelsPerInput = 4;
+		const FIntPoint ProcessedOutputDimension = { InDestinationTexture->GetDesc().Extent.X / PixelsPerInput,InDestinationTexture->GetDesc().Extent.Y };
+		GroupCount = FComputeShaderUtils::GetGroupCount(ProcessedOutputDimension, FComputeShaderUtils::kGolden2DGroupSize);
+		FGlobalShaderMap* GlobalShaderMap = GetGlobalShaderMap(GMaxRHIFeatureLevel);
+		TShaderMapRef<FRGB8BitToRGBA8CS> ComputeShader(GlobalShaderMap, PermutationVector);
+		FRGB8BitToRGBA8CS::FParameters* Parameters = ComputeShader->AllocateAndSetParameters(GraphBuilder, InputRGGBuffer, OutputResource, ElementsPerRow);
+
+
+		FComputeShaderUtils::AddPass(
+			GraphBuilder
+			, RDG_EVENT_NAME("RGB8BitToRGBA8")
+			, ComputeShader
+			, Parameters
+			, GroupCount);
+		break;
+	}
+	case ERivermaxMediaSourePixelFormat::RGB_10bit:
+	{
+		constexpr uint32 PixelPerGroup = 4;
+		constexpr uint32 BytesPerGroup = (PixelPerGroup * 3 * 10) / 8;
+		const uint32 HorizontalByteCount = (InDestinationTexture->GetDesc().Extent.X / PixelPerGroup) * BytesPerGroup;
+		const int32 BytesPerElement = sizeof(FRGBToRGB10BitCS::FRGB10BitBuffer);
+		const int32 ElementsPerRow = HorizontalByteCount / BytesPerElement;
+		const int32 ElementCount = ElementsPerRow * InDestinationTexture->GetDesc().Extent.Y;
+		
+		FRGB10BitToRGBA10CS::FPermutationDomain PermutationVector;
+		PermutationVector.Set<FRGB10BitToRGBA10CS::FSRGBToLinear>(bDoSRGBToLinear);
+		
+		FRDGBufferRef InputRGGBuffer = CreateStructuredBuffer(GraphBuilder, TEXT("RivermaxInputBuffer"), BytesPerElement, ElementCount, SamplePtr->GetBuffer(), BytesPerElement * ElementCount);
+		constexpr int32 PixelsPerInput = 16;
+		const FIntPoint ProcessedOutputDimension = { InDestinationTexture->GetDesc().Extent.X / PixelsPerInput,InDestinationTexture->GetDesc().Extent.Y};
+		GroupCount = FComputeShaderUtils::GetGroupCount(ProcessedOutputDimension, FComputeShaderUtils::kGolden2DGroupSize);
+		FGlobalShaderMap* GlobalShaderMap = GetGlobalShaderMap(GMaxRHIFeatureLevel);
+		TShaderMapRef<FRGB10BitToRGBA10CS> ComputeShader(GlobalShaderMap, PermutationVector);
+		FRGB10BitToRGBA10CS::FParameters* Parameters = ComputeShader->AllocateAndSetParameters(GraphBuilder, InputRGGBuffer, OutputResource, ElementsPerRow);
+
+
+		FComputeShaderUtils::AddPass(
+			GraphBuilder
+			, RDG_EVENT_NAME("RGB10BitToRGBA")
+			, ComputeShader
+			, Parameters
+			, GroupCount);
+		break;
+	}
+	default:
+	{
+		ensureMsgf(false, TEXT("Unhandled pixel format (%d) given to Rivermax MediaSample converter"), InputPixelFormat);
+		return false;
+	}
+	}
+
+	GraphBuilder.Execute();
+
+	return true;
+}
+
+uint32 FRivermaxMediaTextureSampleConverter::GetConverterInfoFlags() const
+{
+	return IMediaTextureSampleConverter::ConverterInfoFlags_NeedUAVOutputTexture;
+}
diff --git a/Engine/Plugins/VirtualProduction/Rivermax/RivermaxMedia/Source/RivermaxMedia/Private/RivermaxMediaTextureSampleConverter.h b/Engine/Plugins/VirtualProduction/Rivermax/RivermaxMedia/Source/RivermaxMedia/Private/RivermaxMediaTextureSampleConverter.h
new file mode 100644
--- /dev/null
+++ b/Engine/Plugins/VirtualProduction/Rivermax/RivermaxMedia/Source/RivermaxMedia/Private/RivermaxMediaTextureSampleConverter.h
@@ -0,0 +1,37 @@
+// Copyright Epic Games, Inc. All Rights Reserved.
+
+#pragma once
+
+#include "IMediaTextureSampleConverter.h"
+
+#include "RivermaxMediaSource.h"
+
+
+class FRivermaxMediaTextureSample;
+
+/**
+ * 
+ */
+class FRivermaxMediaTextureSampleConverter : public IMediaTextureSampleConverter
+{
+public:
+
+	/** Configures settings to convert incoming sample */
+	void Setup(ERivermaxMediaSourePixelFormat InPixelFormat, TWeakPtr<FRivermaxMediaTextureSample> InSample, bool bInDoSRGBToLinear);
+
+	//~ Begin IMediaTextureSampleConverter interface
+	virtual bool Convert(FTexture2DRHIRef& InDstTexture, const FConversionHints& Hints) override;
+	virtual uint32 GetConverterInfoFlags() const override;
+	//~ End IMediaTextureSampleConverter interface
+
+private:
+
+	/** Incoming pixel format used to know how to handle the input buffer */
+	ERivermaxMediaSourePixelFormat InputPixelFormat = ERivermaxMediaSourePixelFormat::RGB_10bit;
+
+	/** Sample holding the buffer to convert */
+	TWeakPtr<FRivermaxMediaTextureSample> Sample;
+
+	/** Whether shader should do a sRGB to linear conversion before writing out to the texture */
+	bool bDoSRGBToLinear = false;
+};
diff --git a/Engine/Plugins/VirtualProduction/Rivermax/RivermaxMedia/Source/RivermaxMedia/Public/RivermaxMediaOutput.h b/Engine/Plugins/VirtualProduction/Rivermax/RivermaxMedia/Source/RivermaxMedia/Public/RivermaxMediaOutput.h
--- a/Engine/Plugins/VirtualProduction/Rivermax/RivermaxMedia/Source/RivermaxMedia/Public/RivermaxMediaOutput.h
+++ b/Engine/Plugins/VirtualProduction/Rivermax/RivermaxMedia/Source/RivermaxMedia/Public/RivermaxMediaOutput.h
@@ -60,10 +60,10 @@ public:
 	ERivermaxMediaOutputPixelFormat PixelFormat = ERivermaxMediaOutputPixelFormat::PF_8BIT_YUV;
 
 	UPROPERTY(EditAnywhere, BlueprintReadWrite, Category = "Format")
-	FString SourceAddress;
+	FString InterfaceAddress;
 
 	UPROPERTY(EditAnywhere, BlueprintReadWrite, Category = "Format")
-	FString DestinationAddress;
+	FString StreamAddress;
 
 	UPROPERTY(EditAnywhere, BlueprintReadWrite, Category = "Format")
 	int32 Port = 50000;
diff --git a/Engine/Plugins/VirtualProduction/Rivermax/RivermaxMedia/Source/RivermaxMedia/Public/RivermaxMediaSource.h b/Engine/Plugins/VirtualProduction/Rivermax/RivermaxMedia/Source/RivermaxMedia/Public/RivermaxMediaSource.h
--- a/Engine/Plugins/VirtualProduction/Rivermax/RivermaxMedia/Source/RivermaxMedia/Public/RivermaxMediaSource.h
+++ b/Engine/Plugins/VirtualProduction/Rivermax/RivermaxMedia/Source/RivermaxMedia/Public/RivermaxMediaSource.h
@@ -12,10 +12,12 @@
  * Native data format.
  */
 UENUM()
-enum class ERivermaxMediaSourceColorFormat : uint8
+enum class ERivermaxMediaSourePixelFormat : uint8
 {
-	YUV2_8bit UMETA(DisplayName = "8bit YUV"),
-	//Todo Add support for 10bit YUV and 8/10 bit RGB
+	YUV422_8bit UMETA(DisplayName = "8bit YUV422"),
+	RGB_8bit UMETA(DisplayName = "8bit RGB"),
+	RGB_10bit UMETA(DisplayName = "10bit RGB"),
+	//Todo Add support for 10bit YUV and float16 rgb
 };
 
 /**
@@ -28,8 +30,6 @@ class RIVERMAXMEDIA_API URivermaxMediaSource : public UTimeSynchronizableMediaSo
 
 public:
 
-	//todo proper device configuration like Aja and Blackmagic
-
 	UPROPERTY(EditAnywhere, BlueprintReadWrite, Category="Format")
 	FIntPoint Resolution = {1920, 1080};
 	
@@ -37,17 +37,23 @@ public:
 	FFrameRate FrameRate = {24,1};
 	
 	UPROPERTY(EditAnywhere, BlueprintReadWrite, Category = "Format")
-	ERivermaxMediaSourceColorFormat PixelFormat = ERivermaxMediaSourceColorFormat::YUV2_8bit;
+	ERivermaxMediaSourePixelFormat PixelFormat = ERivermaxMediaSourePixelFormat::YUV422_8bit;
 
+	/** Network card interface to use to receive data */
 	UPROPERTY(EditAnywhere, BlueprintReadWrite, Category = "Format")
-	FString SourceAddress;
+	FString InterfaceAddress;
 
+	/** IP address where incoming stream is coming from.  */
 	UPROPERTY(EditAnywhere, BlueprintReadWrite, Category = "Format")
-	FString DestinationAddress;
+	FString StreamAddress;
 
 	UPROPERTY(EditAnywhere, BlueprintReadWrite, Category = "Format")
 	int32 Port = 50000;
 
+	/** Whether the video input is in sRGB color space.If true, sRGBToLinear will be done on incoming pixels before writing to media texture */
+	UPROPERTY(EditAnywhere, BlueprintReadWrite, Category = "Video")
+	bool bIsSRGBInput = false;
+
 public:
 	//~ Begin IMediaOptions interface
 	virtual bool GetMediaOption(const FName& Key, bool DefaultValue) const override;
diff --git a/Engine/Plugins/VirtualProduction/Rivermax/RivermaxMedia/Source/RivermaxMediaEditor/RivermaxMediaEditor.build.cs b/Engine/Plugins/VirtualProduction/Rivermax/RivermaxMedia/Source/RivermaxMediaEditor/RivermaxMediaEditor.build.cs
--- a/Engine/Plugins/VirtualProduction/Rivermax/RivermaxMedia/Source/RivermaxMediaEditor/RivermaxMediaEditor.build.cs
+++ b/Engine/Plugins/VirtualProduction/Rivermax/RivermaxMedia/Source/RivermaxMediaEditor/RivermaxMediaEditor.build.cs
@@ -16,13 +16,16 @@ public class RivermaxMediaEditor : ModuleRules
 				"AssetTools",
 				"Core",
 				"CoreUObject",
+				"EditorStyle",
 				"Engine",
+				"InputCore",
 				"MediaAssets",
 				"MediaIOCore",
 				"MediaIOEditor",
 				"RivermaxMedia",
+				"Slate",
+				"SlateCore",
 				"UnrealEd"
-
 			}
 		);
 	}
