- Adding rivermax input stream stats
- Adding auto restart of player / capture in ndisplay

diff --git a/Engine/Plugins/Runtime/nDisplay/Source/DisplayClusterMedia/Private/Capture/DisplayClusterMediaCaptureBase.cpp b/Engine/Plugins/Runtime/nDisplay/Source/DisplayClusterMedia/Private/Capture/DisplayClusterMediaCaptureBase.cpp
--- a/Engine/Plugins/Runtime/nDisplay/Source/DisplayClusterMedia/Private/Capture/DisplayClusterMediaCaptureBase.cpp
+++ b/Engine/Plugins/Runtime/nDisplay/Source/DisplayClusterMedia/Private/Capture/DisplayClusterMediaCaptureBase.cpp
@@ -3,6 +3,8 @@
 #include "Capture/DisplayClusterMediaCaptureBase.h"
 #include "DisplayClusterMediaHelpers.h"
 #include "DisplayClusterMediaLog.h"
+#include "IDisplayCluster.h"
+#include "IDisplayClusterCallbacks.h"
 
 #include "MediaCapture.h"
 #include "MediaOutput.h"
@@ -18,9 +20,16 @@ FDisplayClusterMediaCaptureBase::FDisplayClusterMediaCaptureBase(const FString&
 {
 	// We expect these to always be valid
 	check(InMediaOutput && InRenderTarget);
+
+	IDisplayCluster::Get().GetCallbacks().OnDisplayClusterPostTick().AddRaw(this, &FDisplayClusterMediaCaptureBase::OnPostClusterTick);
 }
 
 
+FDisplayClusterMediaCaptureBase::~FDisplayClusterMediaCaptureBase()
+{
+	IDisplayCluster::Get().GetCallbacks().OnDisplayClusterPostTick().RemoveAll(this);
+}
+
 void FDisplayClusterMediaCaptureBase::AddReferencedObjects(FReferenceCollector& Collector)
 {
 	Collector.AddReferencedObject(MediaOutput);
@@ -41,22 +50,8 @@ bool FDisplayClusterMediaCaptureBase::StartCapture()
 		{
 			MediaCapture->SetMediaOutput(MediaOutput);
 
-			FMediaCaptureOptions MediaCaptureOptions;
-			MediaCaptureOptions.NumberOfFramesToCapture = -1;
-			MediaCaptureOptions.bResizeSourceBuffer = true;
-			MediaCaptureOptions.bSkipFrameWhenRunningExpensiveTasks = false;
-
-			const bool bCaptureStarted = MediaCapture->CaptureTextureRenderTarget2D(RenderTarget, MediaCaptureOptions);
-			if (bCaptureStarted)
-			{
-				UE_LOG(LogDisplayClusterMedia, Log, TEXT("Started media capture: '%s'"), *GetMediaId());
-			}
-			else
-			{
-				UE_LOG(LogDisplayClusterMedia, Warning, TEXT("Couldn't start media capture '%s'"), *GetMediaId());
-			}
-
-			return bCaptureStarted;
+			bWasCaptureStarted = StartMediaCapture();
+			return bWasCaptureStarted;
 		}
 	}
 
@@ -69,6 +64,7 @@ void FDisplayClusterMediaCaptureBase::StopCapture()
 	{
 		MediaCapture->StopCapture(false);
 		MediaCapture = nullptr;
+		bWasCaptureStarted = false;
 	}
 }
 
@@ -97,3 +93,42 @@ void FDisplayClusterMediaCaptureBase::ExportMediaData(FRHICommandListImmediate&
 		}
 	}
 }
+
+void FDisplayClusterMediaCaptureBase::OnPostClusterTick()
+{
+	if (MediaCapture && bWasCaptureStarted)
+	{
+		if (MediaCapture->GetState() == EMediaCaptureState::Error)
+		{
+			constexpr double Interval = 1.0;
+			const double CurrentTime = FPlatformTime::Seconds();
+			if (CurrentTime - LastRestartTimestamp > Interval)
+			{
+				UE_LOG(LogDisplayClusterMedia, Log, TEXT("MediaCapture '%s' is in error, restarting it."), *GetMediaId());
+
+				StartMediaCapture();
+				LastRestartTimestamp = CurrentTime;
+			}
+		}
+	}
+}
+
+bool FDisplayClusterMediaCaptureBase::StartMediaCapture()
+{
+	FMediaCaptureOptions MediaCaptureOptions;
+	MediaCaptureOptions.NumberOfFramesToCapture = -1;
+	MediaCaptureOptions.bResizeSourceBuffer = true;
+	MediaCaptureOptions.bSkipFrameWhenRunningExpensiveTasks = false;
+
+	const bool bCaptureStarted = MediaCapture->CaptureTextureRenderTarget2D(RenderTarget, MediaCaptureOptions);
+	if (bCaptureStarted)
+	{
+		UE_LOG(LogDisplayClusterMedia, Log, TEXT("Started media capture: '%s'"), *GetMediaId());
+	}
+	else
+	{
+		UE_LOG(LogDisplayClusterMedia, Warning, TEXT("Couldn't start media capture '%s'"), *GetMediaId());
+	}
+
+	return bCaptureStarted;
+}
diff --git a/Engine/Plugins/Runtime/nDisplay/Source/DisplayClusterMedia/Private/Capture/DisplayClusterMediaCaptureBase.h b/Engine/Plugins/Runtime/nDisplay/Source/DisplayClusterMedia/Private/Capture/DisplayClusterMediaCaptureBase.h
--- a/Engine/Plugins/Runtime/nDisplay/Source/DisplayClusterMedia/Private/Capture/DisplayClusterMediaCaptureBase.h
+++ b/Engine/Plugins/Runtime/nDisplay/Source/DisplayClusterMedia/Private/Capture/DisplayClusterMediaCaptureBase.h
@@ -21,6 +21,7 @@ class FDisplayClusterMediaCaptureBase
 {
 public:
 	FDisplayClusterMediaCaptureBase(const FString& MediaId, const FString& ClusterNodeId, UMediaOutput* MediaOutput, UTextureRenderTarget2D* RenderTarget);
+	virtual ~FDisplayClusterMediaCaptureBase();
 
 public:
 	//~ Begin FGCObject interface
@@ -42,6 +43,8 @@ public:
 
 protected:
 	void ExportMediaData(FRHICommandListImmediate& RHICmdList, const FMediaTextureInfo& TextureInfo);
+	void OnPostClusterTick();
+	bool StartMediaCapture();
 
 private:
 	//~ Begin GC by AddReferencedObjects
@@ -49,4 +52,10 @@ private:
 	UMediaCapture*          MediaCapture = nullptr;
 	UTextureRenderTarget2D* RenderTarget = nullptr;
 	//~ End GC by AddReferencedObjects
+
+	// Used to restart media capture in the case it falls in error
+	bool bWasCaptureStarted = false;
+
+	// Used to control the rate at which we try to restart the capture
+	double LastRestartTimestamp = 0;
 };
diff --git a/Engine/Plugins/Runtime/nDisplay/Source/DisplayClusterMedia/Private/Input/DisplayClusterMediaInputBase.cpp b/Engine/Plugins/Runtime/nDisplay/Source/DisplayClusterMedia/Private/Input/DisplayClusterMediaInputBase.cpp
--- a/Engine/Plugins/Runtime/nDisplay/Source/DisplayClusterMedia/Private/Input/DisplayClusterMediaInputBase.cpp
+++ b/Engine/Plugins/Runtime/nDisplay/Source/DisplayClusterMedia/Private/Input/DisplayClusterMediaInputBase.cpp
@@ -35,17 +35,9 @@ bool FDisplayClusterMediaInputBase::Play()
 		MediaPlayer->PlayOnOpen = true;
 		MediaPlayer->OnMediaEvent().AddRaw(this, &FDisplayClusterMediaInputBase::OnMediaEvent);
 
-		const bool bIsPlaying = MediaPlayer->OpenSource(MediaSource);
-		if (bIsPlaying)
-		{
-			UE_LOG(LogDisplayClusterMedia, Log, TEXT("Started playing media: %s"), *GetMediaId());
-		}
-		else
-		{
-			UE_LOG(LogDisplayClusterMedia, Warning, TEXT("Couldn't start playing media: %s"), *GetMediaId());
-		}
-
-		return bIsPlaying;
+		bWasPlayerStarted = MediaPlayer->OpenSource(MediaSource);
+		
+		return bWasPlayerStarted;
 	}
 
 	return false;
@@ -55,6 +47,7 @@ void FDisplayClusterMediaInputBase::Stop()
 {
 	if (MediaPlayer)
 	{
+		bWasPlayerStarted = false;
 		MediaPlayer->Close();
 		MediaPlayer->OnMediaEvent().RemoveAll(this);
 	}
@@ -97,7 +90,7 @@ void FDisplayClusterMediaInputBase::OnMediaEvent(EMediaEvent MediaEvent)
 	{
 	/** The player started connecting to the media source. */
 	case EMediaEvent::MediaConnecting:
-		UE_LOG(LogDisplayClusterMedia, Log, TEXT("Media event for '%s': Connectiong"), *GetMediaId());
+		UE_LOG(LogDisplayClusterMedia, Log, TEXT("Media event for '%s': Connection"), *GetMediaId());
 		break;
 
 	/** A new media source has been opened. */
@@ -108,6 +101,7 @@ void FDisplayClusterMediaInputBase::OnMediaEvent(EMediaEvent MediaEvent)
 	/** The current media source has been closed. */
 	case EMediaEvent::MediaClosed:
 		UE_LOG(LogDisplayClusterMedia, Log, TEXT("Media event for '%s': Closed"), *GetMediaId());
+		OnPlayerClosed();
 		break;
 		
 	/** A media source failed to open. */
@@ -120,3 +114,34 @@ void FDisplayClusterMediaInputBase::OnMediaEvent(EMediaEvent MediaEvent)
 		break;
 	}
 }
+
+bool FDisplayClusterMediaInputBase::StartPlayer()
+{
+	const bool bIsPlaying = MediaPlayer->OpenSource(MediaSource);
+	if (bIsPlaying)
+	{
+		UE_LOG(LogDisplayClusterMedia, Log, TEXT("Started playing media: %s"), *GetMediaId());
+	}
+	else
+	{
+		UE_LOG(LogDisplayClusterMedia, Warning, TEXT("Couldn't start playing media: %s"), *GetMediaId());
+	}
+
+	return bIsPlaying;
+}
+
+void FDisplayClusterMediaInputBase::OnPlayerClosed()
+{
+	if (MediaPlayer && bWasPlayerStarted)
+	{
+		constexpr double Interval = 1.0;
+		const double CurrentTime = FPlatformTime::Seconds();
+		if (CurrentTime - LastRestartTimestamp > Interval)
+		{
+			UE_LOG(LogDisplayClusterMedia, Log, TEXT("MediaPlayer '%s' is in error, restarting it."), *GetMediaId());
+
+			StartPlayer();
+			LastRestartTimestamp = CurrentTime;
+		}
+	}
+}
diff --git a/Engine/Plugins/Runtime/nDisplay/Source/DisplayClusterMedia/Private/Input/DisplayClusterMediaInputBase.h b/Engine/Plugins/Runtime/nDisplay/Source/DisplayClusterMedia/Private/Input/DisplayClusterMediaInputBase.h
--- a/Engine/Plugins/Runtime/nDisplay/Source/DisplayClusterMedia/Private/Input/DisplayClusterMediaInputBase.h
+++ b/Engine/Plugins/Runtime/nDisplay/Source/DisplayClusterMedia/Private/Input/DisplayClusterMediaInputBase.h
@@ -58,6 +58,8 @@ protected:
 
 private:
 	void OnMediaEvent(EMediaEvent MediaEvent);
+	bool StartPlayer();
+	void OnPlayerClosed();
 
 private:
 	//~ Begin GC by AddReferencedObjects
@@ -65,4 +67,10 @@ private:
 	UMediaPlayer*  MediaPlayer = nullptr;
 	UMediaTexture* MediaTexture = nullptr;
 	//~ End GC by AddReferencedObjects
+
+	// Used to restart media player in the case it falls in error
+	bool bWasPlayerStarted = false;
+
+	// Used to control the rate at which we try to restart the player
+	double LastRestartTimestamp = 0;
 };
diff --git a/Engine/Plugins/VirtualProduction/Rivermax/RivermaxCore/Source/RivermaxCore/Private/RivermaxInputStream.cpp b/Engine/Plugins/VirtualProduction/Rivermax/RivermaxCore/Source/RivermaxCore/Private/RivermaxInputStream.cpp
--- a/Engine/Plugins/VirtualProduction/Rivermax/RivermaxCore/Source/RivermaxCore/Private/RivermaxInputStream.cpp
+++ b/Engine/Plugins/VirtualProduction/Rivermax/RivermaxCore/Source/RivermaxCore/Private/RivermaxInputStream.cpp
@@ -13,9 +13,15 @@
 #include <WS2tcpip.h>
 #endif
 
-
 namespace UE::RivermaxCore::Private
 {
+	static TAutoConsoleVariable<int32> CVarUseIncompleteFrames(
+		TEXT("Rivermax.Input.UseIncompleteFrames"),
+		1,
+		TEXT("Whether to keep incomplete frames in the case of missing packets or drop them."),
+		ECVF_Default);
+
+
 #if PLATFORM_SUPPORTS_PRAGMA_PACK
 #pragma pack(push, 1)
 #endif
@@ -244,6 +250,7 @@ namespace UE::RivermaxCore::Private
 		while (bIsActive)
 		{
 			Process_AnyThread();
+			LogStats();
 		}
 
 		if (StreamId)
@@ -298,8 +305,11 @@ namespace UE::RivermaxCore::Private
 
 	void FRivermaxInputStream::ParseChunk(const rmax_in_completion& Completion)
 	{
+		
 		for (uint64 StrideIndex = 0; StrideIndex < Completion.chunk_size; ++StrideIndex)
 		{
+			++StreamStats.ChunksReceived;
+		
 			ensure(Completion.hdr_ptr);
 			if (Completion.hdr_ptr == nullptr)
 			{
@@ -318,7 +328,8 @@ namespace UE::RivermaxCore::Private
 				{
 					if (bIsFirstFrameReceived)
 					{
-						//stats  rx count, received bytes++
+						StreamStats.BytesReceived += Completion.packet_info_arr[StrideIndex].data_size + Completion.packet_info_arr[StrideIndex].hdr_size;
+
 						uint64 LastSequenceNumberIncremented = StreamData.LastSequenceNumber + 1;
 						if (IsExtendedSequenceNumber() == false)
 						{
@@ -332,6 +343,9 @@ namespace UE::RivermaxCore::Private
 							bCanProcessChunk = false;
 							StreamData.WritingOffset = 0;
 							StreamData.ReceivedSize = 0;
+							++StreamStats.TotalPacketLossCount;
+							++StreamStats.FramePacketLossCount;
+							
 							UE_LOG(LogRivermax, Warning, TEXT("Lost %uld packets"), LostPackets);
 						}
 
@@ -370,19 +384,23 @@ namespace UE::RivermaxCore::Private
 								UE_LOG(LogRivermax, Warning, TEXT("Received too much data (%d). Expected %d but received (%d)"), StreamData.ReceivedSize - StreamData.ExpectedSize, StreamData.ExpectedSize, StreamData.ReceivedSize);
 								StreamData.WritingOffset = 0;
 								StreamData.ReceivedSize = 0;
+								++StreamStats.BiggerFramesCount;
 							}
 							else if (HeaderStart->MarkerBit)
 							{
 								if (StreamData.ReceivedSize == StreamData.ExpectedSize)
 								{
 									TRACE_CPUPROFILER_EVENT_SCOPE(FRivermaxInputStream::ProcessingReceivedFrame)
+
+									++StreamStats.FramesReceived;
+
 									FRivermaxInputVideoFrameDescriptor Descriptor;
 									Descriptor.Width = Options.Resolution.X;
 									Descriptor.Height = Options.Resolution.Y;
 									Descriptor.Stride = Options.Stride;
 									FRivermaxInputVideoFrameReception NewFrame;
 									NewFrame.VideoBuffer = StreamData.CurrentFrame;
-									Listener->OnVideoFrameReceived(Descriptor, NewFrame);//todo stride
+									Listener->OnVideoFrameReceived(Descriptor, NewFrame);
 									PrepareNextFrame();
 								}
 								else
@@ -390,7 +408,11 @@ namespace UE::RivermaxCore::Private
 									UE_LOG(LogRivermax, Warning, TEXT("End of frame received (Marker bit) but not enough data was received (missing %d). Expected %d but received (%d)"), StreamData.ExpectedSize - StreamData.ReceivedSize, StreamData.ExpectedSize, StreamData.ReceivedSize);
 									StreamData.WritingOffset = 0;
 									StreamData.ReceivedSize = 0;
+									++StreamStats.InvalidFramesCount;
 								}
+								
+								StreamStats.FramePacketLossCount = 0;
+								++StreamStats.EndOfFrameReceived;
 							}
 						}
 
@@ -408,6 +430,14 @@ namespace UE::RivermaxCore::Private
 					//todo why  Reset RTP header
 					*(uint32_t*)GetRTPHeaderPointer(HeaderPtr) = 0;
 				}
+				else
+				{
+					++StreamStats.InvalidHeadercount;
+				}
+			}
+			else
+			{
+				++StreamStats.EmptyCompletionCount;
 			}
 
 			//todo why?
@@ -433,6 +463,31 @@ namespace UE::RivermaxCore::Private
 		StreamData.ReceivedSize = 0;
 		StreamData.ExpectedSize = Descriptor.VideoBufferSize;
 	}
+
+	void FRivermaxInputStream::LogStats()
+	{
+		static constexpr double LoggingInterval = 1.0;
+
+		const double CurrentTime = FPlatformTime::Seconds();
+		if (CurrentTime - LastLoggingTimestamp >= LoggingInterval)
+		{
+			LastLoggingTimestamp = CurrentTime;
+			UE_LOG(LogRivermax, Verbose, TEXT("Stream %d stats: FrameCount: %llu, EndOfFrame: %llu, Chunks: %llu, Bytes: %llu, PacketLossInFrame: %llu, TotalPacketLoss: %llu, BiggerFrames: %llu, InvalidFrames: %llu, InvalidHeader: %llu, EmptyCompletion: %llu")
+			, StreamId
+			, StreamStats.FramesReceived
+			, StreamStats.EndOfFrameReceived
+			, StreamStats.ChunksReceived
+			, StreamStats.BytesReceived
+			, StreamStats.FramePacketLossCount
+			, StreamStats.TotalPacketLossCount
+			, StreamStats.BiggerFramesCount
+			, StreamStats.InvalidFramesCount
+			, StreamStats.InvalidHeadercount
+			, StreamStats.EmptyCompletionCount
+			);
+		}
+	}
+
 }
 
 
diff --git a/Engine/Plugins/VirtualProduction/Rivermax/RivermaxCore/Source/RivermaxCore/Private/RivermaxInputStream.h b/Engine/Plugins/VirtualProduction/Rivermax/RivermaxCore/Source/RivermaxCore/Private/RivermaxInputStream.h
--- a/Engine/Plugins/VirtualProduction/Rivermax/RivermaxCore/Source/RivermaxCore/Private/RivermaxInputStream.h
+++ b/Engine/Plugins/VirtualProduction/Rivermax/RivermaxCore/Source/RivermaxCore/Private/RivermaxInputStream.h
@@ -29,6 +29,20 @@ namespace UE::RivermaxCore::Private
 		rmax_in_memblock HeaderMemory;
 	};
 
+	struct FInputStreamStats
+	{
+		uint64 InvalidHeadercount = 0;
+		uint64 FramePacketLossCount = 0;
+		uint64 TotalPacketLossCount = 0;
+		uint64 BiggerFramesCount = 0;
+		uint64 InvalidFramesCount = 0;
+		uint64 BytesReceived = 0;
+		uint64 FramesReceived = 0;
+		uint64 ChunksReceived = 0;
+		uint64 EndOfFrameReceived = 0;
+		uint64 EmptyCompletionCount = 0;
+	};
+
 	struct FInputStreamData
 	{
 		uint64 LastSequenceNumber = 0;
@@ -74,6 +88,7 @@ namespace UE::RivermaxCore::Private
 		void ParseChunk(const rmax_in_completion& Completion);
 		bool IsExtendedSequenceNumber() const;
 		void PrepareNextFrame();
+		void LogStats();
 
 	private:
 
@@ -90,9 +105,12 @@ namespace UE::RivermaxCore::Private
 		bool bIsFirstFrameReceived = false;
 		ERivermaxStreamType RivermaxStreamType = ERivermaxStreamType::VIDEO_2110_20_STREAM;
 		FInputStreamData StreamData;
+		FInputStreamStats StreamStats;
 		IRivermaxInputStreamListener* Listener;
 		std::atomic<bool> bIsShuttingDown;
 		TFuture<void> InitTaskFuture;
+		double LastLoggingTimestamp = 0.0;
+
 	};
 }
 
