- Adding SMTPE2110 - YUV10 bit 422 input support for rivermax player

diff --git a/Engine/Plugins/VirtualProduction/Rivermax/RivermaxCore/Shaders/Private/RivermaxShaders.usf b/Engine/Plugins/VirtualProduction/Rivermax/RivermaxCore/Shaders/Private/RivermaxShaders.usf
--- a/Engine/Plugins/VirtualProduction/Rivermax/RivermaxCore/Shaders/Private/RivermaxShaders.usf
+++ b/Engine/Plugins/VirtualProduction/Rivermax/RivermaxCore/Shaders/Private/RivermaxShaders.usf
@@ -29,6 +29,7 @@ struct FYCbCr10Bit422
 	uint DWord4;
 };
 RWStructuredBuffer<FYCbCr10Bit422>	OutYCbCrBuffer;
+StructuredBuffer<FYCbCr10Bit422>	InputYCbCrBuffer;
 
 struct FRGB8BitBuffer
 {
@@ -53,6 +54,7 @@ struct FRGB16fBuffer
 };
 RWStructuredBuffer<FRGB16fBuffer> OutRGB16fBuffer;
 
+
 // Shader to convert from RGB to YUV 422 10 bits packed little endian
 
 [numthreads(8,8,1)]
@@ -261,74 +263,61 @@ void RGB10BitToRGBACS(uint3 ThreadId : SV_DispatchThreadID, uint GroupIndex : SV
 	// Used to normalize pixel values based on 10 bits input
 	const float Mul = 1.0 / 1023.0;
 
-	//InputBuffer[BufferIndex].DWords[0] = ((RGB00.z & 0x03F) << 26) | ((RGB01.x >> 8)   << 24)   | ((RGB00.y & 0xF)  << 20) | ((RGB00.z >> 6)   << 16) | ((RGB00.x & 0x3)  << 14) | ((RGB00.y >> 4)   << 8)								| (RGB00.x >> 2);
-
 	float3 RGB00 = uint3(0,0,0);
 	RGB00.x = ((((InputBuffer[BufferIndex].DWords[00] >>  0)  << 2) | ((InputBuffer[BufferIndex].DWords[00] >> 14) & 0x003)) & 0x3FF) * Mul;
 	RGB00.y = ((((InputBuffer[BufferIndex].DWords[00] >>  8)  << 4) | ((InputBuffer[BufferIndex].DWords[00] >> 20) & 0x00F)) & 0x3FF) * Mul;
 	RGB00.z = ((((InputBuffer[BufferIndex].DWords[00] >> 16)  << 6) | ((InputBuffer[BufferIndex].DWords[00] >> 26) & 0x03F)) & 0x3FF) * Mul;
 
-	//InputBuffer[BufferIndex].DWords[1] = ((RGB01.z & 0x00F) << 28) | ((RGB02.x >> 6)   << 24)   | ((RGB01.y & 0x3)  << 22) | ((RGB01.z >> 4)   << 16)							  | ((RGB01.y >> 2)   << 8)								| (RGB01.x & 0xFF);
 	float3 RGB01 = uint3(0,0,0);
 	RGB01.x = ((((InputBuffer[BufferIndex].DWords[00] >> 24) << 8) | ((InputBuffer[BufferIndex].DWords[01] >> 00) & 0x0FF)) & 0x3FF) * Mul;
 	RGB01.y = ((((InputBuffer[BufferIndex].DWords[01] >>  8) << 2) | ((InputBuffer[BufferIndex].DWords[01] >> 22) & 0x003)) & 0x3FF) * Mul;
 	RGB01.z = ((((InputBuffer[BufferIndex].DWords[01] >> 16) << 4) | ((InputBuffer[BufferIndex].DWords[01] >> 28) & 0x00F)) & 0x3FF) * Mul;
 
-	//InputBuffer[BufferIndex].DWords[2] = ((RGB02.z & 0x003) << 30) | ((RGB03.x >> 4)   << 24)								| ((RGB02.z >> 2)   << 16)							  | ((RGB02.y & 0xFF) << 8)  | ((RGB02.x & 0x3F) << 2)  | (RGB02.y >> 8);
 	float3 RGB02 = uint3(0,0,0);
 	RGB02.x = ((((InputBuffer[BufferIndex].DWords[01] >> 24) << 6) | ((InputBuffer[BufferIndex].DWords[02] >> 02) & 0x03F)) & 0x3FF) * Mul;
 	RGB02.y = ((((InputBuffer[BufferIndex].DWords[02] >>  0) << 8) | ((InputBuffer[BufferIndex].DWords[02] >>  8) & 0x0FF)) & 0x3FF) * Mul;
 	RGB02.z = ((((InputBuffer[BufferIndex].DWords[02] >> 16) << 2) | ((InputBuffer[BufferIndex].DWords[02] >> 30) & 0x002)) & 0x3FF) * Mul;
 	
-	//InputBuffer[BufferIndex].DWords[3] =							  ((RGB04.x >> 2)   << 24)								| ((RGB03.z & 0xFF) << 16) | ((RGB03.y & 0x3F) << 10) | ((RGB03.z >> 8)   << 8)  | ((RGB03.x & 0xF)  << 4)  | (RGB03.y >> 6);
 	float3 RGB03 = uint3(0,0,0);
 	RGB03.x = ((((InputBuffer[BufferIndex].DWords[02] >> 24) << 4) | ((InputBuffer[BufferIndex].DWords[03] >> 04) & 0x00F)) & 0x3FF) * Mul;
 	RGB03.y = ((((InputBuffer[BufferIndex].DWords[03] >> 00) << 6) | ((InputBuffer[BufferIndex].DWords[03] >> 10) & 0x03F)) & 0x3FF) * Mul;
 	RGB03.z = ((((InputBuffer[BufferIndex].DWords[03] >>  8) << 8) | ((InputBuffer[BufferIndex].DWords[03] >> 16) & 0x0FF)) & 0x3FF) * Mul;
 	
-	//InputBuffer[BufferIndex].DWords[4] =							  ((RGB05.x & 0xFF) << 24)   | ((RGB04.z & 0x3F) << 18) | ((RGB05.x >> 8)   << 16) | ((RGB04.y & 0xF)  << 12) | ((RGB04.z >> 6)   << 8)  | ((RGB04.x & 0x3)  << 6)  | (RGB04.y >> 4);
 	float3 RGB04 = uint3(0,0,0);
 	RGB04.x = ((((InputBuffer[BufferIndex].DWords[03] >> 24) << 2) | ((InputBuffer[BufferIndex].DWords[04] >> 06) & 0x003)) & 0x3FF) * Mul;
 	RGB04.y = ((((InputBuffer[BufferIndex].DWords[04] >> 00) << 4) | ((InputBuffer[BufferIndex].DWords[04] >> 12) & 0x00F)) & 0x3FF) * Mul;
 	RGB04.z = ((((InputBuffer[BufferIndex].DWords[04] >>  8) << 6) | ((InputBuffer[BufferIndex].DWords[04] >> 18) & 0x03F)) & 0x3FF) * Mul;
 	
-	//OutRGB10Buffer[Index].DWords[5] = ((RGB06.x & 0x03F) << 26) | ((RGB06.y >> 8)   << 24)   | ((RGB05.z & 0xF)  << 20) | ((RGB06.x >> 6)   << 16) | ((RGB05.y & 0x3)  << 14) | ((RGB05.z >> 4)   << 8)								| (RGB05.y >> 2);
 	float3 RGB05;
 	RGB05.x = ((((InputBuffer[BufferIndex].DWords[04] >> 16) << 8) | ((InputBuffer[BufferIndex].DWords[04] >> 24) & 0x0FF)) & 0x3FF) * Mul;
 	RGB05.y = ((((InputBuffer[BufferIndex].DWords[05] >> 00) << 2) | ((InputBuffer[BufferIndex].DWords[05] >> 14) & 0x003)) & 0x3FF) * Mul;
 	RGB05.z = ((((InputBuffer[BufferIndex].DWords[05] >>  8) << 4) | ((InputBuffer[BufferIndex].DWords[05] >> 20) & 0x00F)) & 0x3FF) * Mul;
 
-	//OutRGB10Buffer[Index].DWords[6] = ((RGB07.x & 0x00F) << 28) | ((RGB07.y >> 6)   << 24)   | ((RGB06.z & 0x3)  << 22) | ((RGB07.x >> 4)   << 16)							  | ((RGB06.z >> 2)   << 8)								| (RGB06.y & 0xFF);
 	float3 RGB06 = uint3(0,0,0);
 	RGB06.x = ((((InputBuffer[BufferIndex].DWords[05] >> 16) << 6) | ((InputBuffer[BufferIndex].DWords[05] >> 26) & 0x03F)) & 0x3FF) * Mul;
 	RGB06.y = ((((InputBuffer[BufferIndex].DWords[05] >> 24) << 8) | ((InputBuffer[BufferIndex].DWords[06] >> 00) & 0x0FF)) & 0x3FF) * Mul;
 	RGB06.z = ((((InputBuffer[BufferIndex].DWords[06] >>  8) << 2) | ((InputBuffer[BufferIndex].DWords[06] >> 22) & 0x003)) & 0x3FF) * Mul;
 
-	//OutRGB10Buffer[Index].DWords[7] = ((RGB08.x & 0x003) << 30) | ((RGB08.y >> 4)   << 24)								| ((RGB08.x >> 2)   << 16)							  | ((RGB07.z & 0xFF) << 8)  | ((RGB07.y & 0x3F) << 2)  | (RGB07.z >> 8);
 	float3 RGB07 = uint3(0,0,0);
 	RGB07.x = ((((InputBuffer[BufferIndex].DWords[06] >> 16) << 4) | ((InputBuffer[BufferIndex].DWords[06] >> 28) & 0x00F)) & 0x3FF) * Mul;
 	RGB07.y = ((((InputBuffer[BufferIndex].DWords[06] >> 24) << 6) | ((InputBuffer[BufferIndex].DWords[07] >> 02) & 0x03F)) & 0x3FF) * Mul;
 	RGB07.z = ((((InputBuffer[BufferIndex].DWords[07] >>  0) << 8) | ((InputBuffer[BufferIndex].DWords[07] >>  8) & 0x0FF)) & 0x3FF) * Mul;
 
-	//OutRGB10Buffer[Index].DWords[8] =							  ((RGB09.y >> 2)   << 24)								| ((RGB09.x & 0xFF) << 16) | ((RGB08.z & 0x3F) << 10) | ((RGB09.x >> 8)   << 8)  | ((RGB08.y & 0xF)  << 4)  | (RGB08.z >> 6);
 	float3 RGB08 = uint3(0,0,0);
 	RGB08.x = ((((InputBuffer[BufferIndex].DWords[07] >> 16) << 2) | ((InputBuffer[BufferIndex].DWords[07] >> 30) & 0x003)) & 0x3FF) * Mul;
 	RGB08.y = ((((InputBuffer[BufferIndex].DWords[07] >> 24) << 4) | ((InputBuffer[BufferIndex].DWords[ 8] >> 04) & 0x00F)) & 0x3FF) * Mul;
 	RGB08.z = ((((InputBuffer[BufferIndex].DWords[ 8] >>  0) << 6) | ((InputBuffer[BufferIndex].DWords[ 8] >> 10) & 0x03F)) & 0x3FF) * Mul;
 
-	//OutRGB10Buffer[Index].DWords[9] =							  ((RGB10.y & 0xFF) << 24)   | ((RGB10.x & 0x3F) << 18) | ((RGB10.y >> 8)   << 16) | ((RGB09.z & 0xF)  << 12) | ((RGB10.x >> 6)   << 8)  | ((RGB09.y & 0x3)  << 6)  | (RGB09.z >> 4);
 	float3 RGB09 = uint3(0,0,0);
 	RGB09.x = ((((InputBuffer[BufferIndex].DWords[ 8] >>  8) << 8) | ((InputBuffer[BufferIndex].DWords[ 8] >> 16) & 0x0FF)) & 0x3FF) * Mul;
 	RGB09.y = ((((InputBuffer[BufferIndex].DWords[ 8] >> 24) << 2) | ((InputBuffer[BufferIndex].DWords[ 9] >> 06) & 0x003)) & 0x3FF) * Mul;
 	RGB09.z = ((((InputBuffer[BufferIndex].DWords[ 9] >>  0) << 4) | ((InputBuffer[BufferIndex].DWords[ 9] >> 12) & 0x00F)) & 0x3FF) * Mul;
 
-	//OutRGB10Buffer[Index].DWords[10] = ((RGB11.y & 0x03F) << 26) | ((RGB11.z >> 8)   << 24)  | ((RGB11.x & 0xF)  << 20) | ((RGB11.y >> 6)   << 16) | ((RGB10.z & 0x3)  << 14) | ((RGB11.x >> 4)   << 8)								| (RGB10.z >> 2);
 	float3 RGB10 = uint3(0,0,0);
 	RGB10.x = ((((InputBuffer[BufferIndex].DWords[ 9] >>  8) << 6) | ((InputBuffer[BufferIndex].DWords[ 9] >> 18) & 0x03F)) & 0x3FF) * Mul;
 	RGB10.y = ((((InputBuffer[BufferIndex].DWords[ 9] >> 16) << 8) | ((InputBuffer[BufferIndex].DWords[ 9] >> 24) & 0x0FF)) & 0x3FF) * Mul;
 	RGB10.z = ((((InputBuffer[BufferIndex].DWords[10] >>  0) << 2) | ((InputBuffer[BufferIndex].DWords[10] >> 14) & 0x003)) & 0x3FF) * Mul;
 
-	//OutRGB10Buffer[Index].DWords[11] = ((RGB12.y & 0x00F) << 28) | ((RGB12.z >> 6)   << 24)  | ((RGB12.x & 0x3)  << 22) | ((RGB12.y >> 4)   << 16)							  | ((RGB12.x >> 2)   << 8)								| (RGB11.z & 0xFF);
 	float3 RGB11 = uint3(0,0,0);
 	RGB11.x = ((((InputBuffer[BufferIndex].DWords[10] >>  8) << 4) | ((InputBuffer[BufferIndex].DWords[10] >> 20) & 0x00F)) & 0x3FF) * Mul;
 	RGB11.y = ((((InputBuffer[BufferIndex].DWords[10] >> 16) << 6) | ((InputBuffer[BufferIndex].DWords[10] >> 26) & 0x03F)) & 0x3FF) * Mul;
@@ -340,13 +329,11 @@ void RGB10BitToRGBACS(uint3 ThreadId : SV_DispatchThreadID, uint GroupIndex : SV
 	RGB12.y = ((((InputBuffer[BufferIndex].DWords[11] >> 16) << 4) | ((InputBuffer[BufferIndex].DWords[11] >> 28) & 0x00F)) & 0x3FF) * Mul;
 	RGB12.z = ((((InputBuffer[BufferIndex].DWords[11] >> 24) << 6) | ((InputBuffer[BufferIndex].DWords[12] >> 02) & 0x03F)) & 0x3FF) * Mul;
 
-	//OutRGB10Buffer[Index].DWords[13] =							   ((RGB14.z >> 2)   << 24)								| ((RGB14.y & 0xFF) << 16) | ((RGB14.x & 0x3F) << 10) | ((RGB14.y >> 8)   << 8)  | ((RGB13.z & 0xF)  << 4)  | (RGB14.x >> 6);
 	float3 RGB13 = uint3(0,0,0);
 	RGB13.x = ((((InputBuffer[BufferIndex].DWords[12] >>  0) << 8) | ((InputBuffer[BufferIndex].DWords[12] >>  8) & 0x0FF)) & 0x3FF) * Mul;
 	RGB13.y = ((((InputBuffer[BufferIndex].DWords[12] >> 16) << 2) | ((InputBuffer[BufferIndex].DWords[12] >> 30) & 0x003)) & 0x3FF) * Mul;
 	RGB13.z = ((((InputBuffer[BufferIndex].DWords[12] >> 24) << 4) | ((InputBuffer[BufferIndex].DWords[13] >> 04) & 0x00F)) & 0x3FF) * Mul;
 
-	//OutRGB10Buffer[Index].DWords[14] =							   ((RGB15.z & 0xFF) << 24)  | ((RGB15.y & 0x3F) << 18) | ((RGB15.z >> 8)   << 16) | ((RGB15.x & 0xF)  << 12) | ((RGB15.y >> 6)   << 8)  | ((RGB14.z & 0x3)  << 6)  | (RGB15.x >> 4);
 	float3 RGB14 = uint3(0,0,0);
 	RGB14.x = ((((InputBuffer[BufferIndex].DWords[13] >>  0) << 6) | ((InputBuffer[BufferIndex].DWords[13] >> 10) & 0x03F)) & 0x3FF) * Mul;
 	RGB14.y = ((((InputBuffer[BufferIndex].DWords[13] >>  8) << 8) | ((InputBuffer[BufferIndex].DWords[13] >> 16) & 0x0FF)) & 0x3FF) * Mul;
@@ -393,3 +380,64 @@ void RGB10BitToRGBACS(uint3 ThreadId : SV_DispatchThreadID, uint GroupIndex : SV
 	OutTexture[uint2(BaseU + 14, BaseV)] = float4(RGB14, 1.0);
 	OutTexture[uint2(BaseU + 15, BaseV)] = float4(RGB15, 1.0);
 }
+
+
+// Shader to convert from YUV 422 10 bits packed little endian to RGBA
+
+[numthreads(8,8,1)]
+void YUV10Bit422ToRGBACS(uint3 ThreadId : SV_DispatchThreadID, uint GroupIndex : SV_GroupIndex)
+{
+	// Pixels we treat in shader is pretty much hardcoded to 8 so no parameter has been used
+	const uint OutputTexturePixelsPerThread = 8;
+
+	// Scale U (x index) based on the number of pixels we treat per thread
+	const uint BaseU = ThreadId.x * OutputTexturePixelsPerThread;
+	const uint BaseV = ThreadId.y;
+
+	// Input buffer indexing is done in one dimension so adjust using thread Id and elements per row
+	const uint BufferIndex = ThreadId.y * HorizontalElementCount + ThreadId.x;
+
+	// Used to normalize pixel values based on 10 bits input
+	const float Mul = 1.0 / 1023.0;
+
+	float Y0 = ((((InputYCbCrBuffer[BufferIndex].DWord0 >>  8)  << 4) | ((InputYCbCrBuffer[BufferIndex].DWord0 >> 20) & 0x00F)) & 0x3FF) * Mul;
+	float U0 = ((((InputYCbCrBuffer[BufferIndex].DWord0 >>  0)  << 2) | ((InputYCbCrBuffer[BufferIndex].DWord0 >> 14) & 0x003)) & 0x3FF) * Mul;
+	float V0 = ((((InputYCbCrBuffer[BufferIndex].DWord0 >> 16)  << 6) | ((InputYCbCrBuffer[BufferIndex].DWord0 >> 26) & 0x03F)) & 0x3FF) * Mul;
+	float Y1 = ((((InputYCbCrBuffer[BufferIndex].DWord0 >> 24)  << 8) | ((InputYCbCrBuffer[BufferIndex].DWord1 >> 00) & 0x0FF)) & 0x3FF) * Mul;
+	float Y2 = ((((InputYCbCrBuffer[BufferIndex].DWord1 >> 16)  << 4) | ((InputYCbCrBuffer[BufferIndex].DWord1 >> 28) & 0x00F)) & 0x3FF) * Mul;
+	float U2 = ((((InputYCbCrBuffer[BufferIndex].DWord1 >>  8)  << 2) | ((InputYCbCrBuffer[BufferIndex].DWord1 >> 22) & 0x003)) & 0x3FF) * Mul;
+	float V2 = ((((InputYCbCrBuffer[BufferIndex].DWord1 >> 24)  << 6) | ((InputYCbCrBuffer[BufferIndex].DWord2 >> 02) & 0x03F)) & 0x3FF) * Mul;
+	float Y3 = ((((InputYCbCrBuffer[BufferIndex].DWord2 >> 00)  << 8) | ((InputYCbCrBuffer[BufferIndex].DWord2 >>  8) & 0x0FF)) & 0x3FF) * Mul;
+	float Y4 = ((((InputYCbCrBuffer[BufferIndex].DWord2 >> 24)  << 4) | ((InputYCbCrBuffer[BufferIndex].DWord3 >> 04) & 0x00F)) & 0x3FF) * Mul;
+	float U4 = ((((InputYCbCrBuffer[BufferIndex].DWord2 >> 16)  << 2) | ((InputYCbCrBuffer[BufferIndex].DWord2 >> 30) & 0x003)) & 0x3FF) * Mul;
+	float V4 = ((((InputYCbCrBuffer[BufferIndex].DWord3 >> 00)  << 6) | ((InputYCbCrBuffer[BufferIndex].DWord3 >> 10) & 0x03F)) & 0x3FF) * Mul;
+	float Y5 = ((((InputYCbCrBuffer[BufferIndex].DWord3 >>  8)  << 8) | ((InputYCbCrBuffer[BufferIndex].DWord3 >> 16) & 0x0FF)) & 0x3FF) * Mul;
+	float Y6 = ((((InputYCbCrBuffer[BufferIndex].DWord4 >>  0)  << 4) | ((InputYCbCrBuffer[BufferIndex].DWord3 >> 12) & 0x00F)) & 0x3FF) * Mul;
+	float U6 = ((((InputYCbCrBuffer[BufferIndex].DWord3 >> 24)  << 2) | ((InputYCbCrBuffer[BufferIndex].DWord4 >> 06) & 0x003)) & 0x3FF) * Mul;
+	float V6 = ((((InputYCbCrBuffer[BufferIndex].DWord4 >>  8)  << 6) | ((InputYCbCrBuffer[BufferIndex].DWord4 >> 18) & 0x03F)) & 0x3FF) * Mul;
+	float Y7 = ((((InputYCbCrBuffer[BufferIndex].DWord4 >> 16)  << 8) | ((InputYCbCrBuffer[BufferIndex].DWord4 >> 24) & 0x0FF)) & 0x3FF) * Mul;
+
+	#if DO_SRGB_TO_LINEAR
+	int DoSRGBToLinear = 1;
+	#else
+	int DoSRGBToLinear = 0;
+	#endif
+
+	float3 RGB00 = YuvToRgb(float3(Y0, U0, V0), ColorTransform, DoSRGBToLinear);
+	float3 RGB01 = YuvToRgb(float3(Y1, U0, V0), ColorTransform, DoSRGBToLinear);
+	float3 RGB02 = YuvToRgb(float3(Y2, U2, V2), ColorTransform, DoSRGBToLinear);
+	float3 RGB03 = YuvToRgb(float3(Y3, U2, V2), ColorTransform, DoSRGBToLinear);
+	float3 RGB04 = YuvToRgb(float3(Y4, U4, V4), ColorTransform, DoSRGBToLinear);
+	float3 RGB05 = YuvToRgb(float3(Y5, U4, V4), ColorTransform, DoSRGBToLinear);
+	float3 RGB06 = YuvToRgb(float3(Y6, U6, V6), ColorTransform, DoSRGBToLinear);
+	float3 RGB07 = YuvToRgb(float3(Y7, U6, V6), ColorTransform, DoSRGBToLinear);
+
+	OutTexture[uint2(BaseU + 0, BaseV)] = float4(RGB00, 1.0);
+	OutTexture[uint2(BaseU + 1, BaseV)] = float4(RGB01, 1.0);
+	OutTexture[uint2(BaseU + 2, BaseV)] = float4(RGB02, 1.0);
+	OutTexture[uint2(BaseU + 3, BaseV)] = float4(RGB03, 1.0);
+	OutTexture[uint2(BaseU + 4, BaseV)] = float4(RGB04, 1.0);
+	OutTexture[uint2(BaseU + 5, BaseV)] = float4(RGB05, 1.0);
+	OutTexture[uint2(BaseU + 6, BaseV)] = float4(RGB06, 1.0);
+	OutTexture[uint2(BaseU + 7, BaseV)] = float4(RGB07, 1.0);
+}
diff --git a/Engine/Plugins/VirtualProduction/Rivermax/RivermaxCore/Source/RivermaxRendering/Private/RivermaxShaders.cpp b/Engine/Plugins/VirtualProduction/Rivermax/RivermaxCore/Source/RivermaxRendering/Private/RivermaxShaders.cpp
--- a/Engine/Plugins/VirtualProduction/Rivermax/RivermaxCore/Source/RivermaxRendering/Private/RivermaxShaders.cpp
+++ b/Engine/Plugins/VirtualProduction/Rivermax/RivermaxCore/Source/RivermaxRendering/Private/RivermaxShaders.cpp
@@ -8,6 +8,23 @@
 
 namespace UE::RivermaxShaders
 {
+	/* FYUV10Bit422ToRGBACS shader
+	 *****************************************************************************/
+
+	IMPLEMENT_GLOBAL_SHADER(FYUV10Bit422ToRGBACS, "/Plugin/RivermaxCore/Private/RivermaxShaders.usf", "YUV10Bit422ToRGBACS", SF_Compute);
+
+	FYUV10Bit422ToRGBACS::FParameters* FYUV10Bit422ToRGBACS::AllocateAndSetParameters(FRDGBuilder& GraphBuilder, FRDGBufferRef YUVBuffer, FRDGTextureRef OutputTexture, const FMatrix& ColorTransform, const FVector& YUVOffset, int32 BufferElementsPerRow)
+	{
+		FYUV10Bit422ToRGBACS::FParameters* Parameters = GraphBuilder.AllocParameters<FYUV10Bit422ToRGBACS::FParameters>();
+
+		Parameters->InputYCbCrBuffer = GraphBuilder.CreateSRV(YUVBuffer);
+		Parameters->ColorTransform = (FMatrix44f)MediaShaders::CombineColorTransformAndOffset(ColorTransform, YUVOffset);
+		Parameters->HorizontalElementCount = BufferElementsPerRow;
+		Parameters->OutTexture = GraphBuilder.CreateUAV(OutputTexture, ERDGUnorderedAccessViewFlags::None);
+
+		return Parameters;
+	}
+
 	/* FRGBToYUV10Bit422LittleEndianCS shader
 	 *****************************************************************************/
 
diff --git a/Engine/Plugins/VirtualProduction/Rivermax/RivermaxCore/Source/RivermaxRendering/Public/RivermaxShaders.h b/Engine/Plugins/VirtualProduction/Rivermax/RivermaxCore/Source/RivermaxRendering/Public/RivermaxShaders.h
--- a/Engine/Plugins/VirtualProduction/Rivermax/RivermaxCore/Source/RivermaxRendering/Public/RivermaxShaders.h
+++ b/Engine/Plugins/VirtualProduction/Rivermax/RivermaxCore/Source/RivermaxRendering/Public/RivermaxShaders.h
@@ -13,6 +13,48 @@
 
 namespace UE::RivermaxShaders
 {
+	/**
+	 * Compute shader to convert packed 8 bits RGB to RGBA 8bits
+	 */
+	class RIVERMAXRENDERING_API FYUV10Bit422ToRGBACS : public FGlobalShader
+	{
+	public:
+
+		//Structure definition to match StructuredBuffer in RivermaxShaders.usf
+		struct FYUV10Bit422LEBuffer
+		{
+			uint32 DWord0;
+			uint32 DWord1;
+			uint32 DWord2;
+			uint32 DWord3;
+			uint32 DWord4;
+		};
+
+		DECLARE_GLOBAL_SHADER(FYUV10Bit422ToRGBACS);
+		SHADER_USE_PARAMETER_STRUCT(FYUV10Bit422ToRGBACS, FGlobalShader);
+
+		class FSRGBToLinear : SHADER_PERMUTATION_BOOL("DO_SRGB_TO_LINEAR");
+		using FPermutationDomain = TShaderPermutationDomain<FSRGBToLinear>;
+
+		BEGIN_SHADER_PARAMETER_STRUCT(FParameters, )
+			SHADER_PARAMETER_RDG_BUFFER_SRV(StructuredBuffer<FYUV10Bit422LEBuffer>, InputYCbCrBuffer)
+			SHADER_PARAMETER(FMatrix44f, ColorTransform)	
+			SHADER_PARAMETER(uint32, HorizontalElementCount)
+			SHADER_PARAMETER_RDG_TEXTURE_UAV(RWTexture2D, OutTexture)
+		END_SHADER_PARAMETER_STRUCT()
+
+	public:
+
+		// Called by the engine to determine which permutations to compile for this shader
+		static bool ShouldCompilePermutation(const FGlobalShaderPermutationParameters& Parameters)
+		{
+			return IsFeatureLevelSupported(Parameters.Platform, ERHIFeatureLevel::SM5);
+		}
+
+		/** Allocates and setup shader parameter in the incoming graph builder */
+		FYUV10Bit422ToRGBACS::FParameters* AllocateAndSetParameters(FRDGBuilder& GraphBuilder, FRDGBufferRef YUVBuffer, FRDGTextureRef OutputTexture, const FMatrix& ColorTransform, const FVector& YUVOffset, int32 BufferElementsPerRow);
+	};
+
 	/**
 	 * Compute shader to convert RGB to packed YUV422 10 bits little endian
 	 *
diff --git a/Engine/Plugins/VirtualProduction/Rivermax/RivermaxMedia/Source/RivermaxMedia/Private/RivermaxMediaPlayer.cpp b/Engine/Plugins/VirtualProduction/Rivermax/RivermaxMedia/Source/RivermaxMedia/Private/RivermaxMediaPlayer.cpp
--- a/Engine/Plugins/VirtualProduction/Rivermax/RivermaxMedia/Source/RivermaxMedia/Private/RivermaxMediaPlayer.cpp
+++ b/Engine/Plugins/VirtualProduction/Rivermax/RivermaxMedia/Source/RivermaxMedia/Private/RivermaxMediaPlayer.cpp
@@ -291,6 +291,12 @@ namespace UE::RivermaxMedia
 			StreamOptions.Stride = StreamOptions.Resolution.X / 2 * 4; // 4 bytes for a group of 2 pixels
 			break;
 		}
+		case ERivermaxMediaSourePixelFormat::YUV422_10bit:
+		{
+			StreamOptions.PixelFormat = ERivermaxOutputPixelFormat::RMAX_10BIT_YCBCR;
+			StreamOptions.Stride = StreamOptions.Resolution.X / 2 * 5; // 5 bytes for a group of 2 pixels (40bits / 8 = 5)
+			break;
+		}
 		case ERivermaxMediaSourePixelFormat::RGB_8bit:
 		{
 			StreamOptions.PixelFormat = ERivermaxOutputPixelFormat::RMAX_8BIT_RGB;
diff --git a/Engine/Plugins/VirtualProduction/Rivermax/RivermaxMedia/Source/RivermaxMedia/Private/RivermaxMediaTextureSample.cpp b/Engine/Plugins/VirtualProduction/Rivermax/RivermaxMedia/Source/RivermaxMedia/Private/RivermaxMediaTextureSample.cpp
--- a/Engine/Plugins/VirtualProduction/Rivermax/RivermaxMedia/Source/RivermaxMedia/Private/RivermaxMediaTextureSample.cpp
+++ b/Engine/Plugins/VirtualProduction/Rivermax/RivermaxMedia/Source/RivermaxMedia/Private/RivermaxMediaTextureSample.cpp
@@ -33,12 +33,16 @@ bool FRivermaxMediaTextureSample::ConfigureSample(uint32 InWidth, uint32 InHeigh
 	switch (InSampleFormat)
 	{
 		case ERivermaxMediaSourePixelFormat::RGB_10bit:
+			// Falls through
+		case ERivermaxMediaSourePixelFormat::YUV422_10bit:
 		{
 			VideoSampleFormat = EMediaTextureSampleFormat::CharBGR10A2;
 			break;
 		}
 		case ERivermaxMediaSourePixelFormat::YUV422_8bit:
+			// Falls through
 		case ERivermaxMediaSourePixelFormat::RGB_8bit:
+			// Falls through
 		default:
 		{
 			VideoSampleFormat = EMediaTextureSampleFormat::CharBGRA;
diff --git a/Engine/Plugins/VirtualProduction/Rivermax/RivermaxMedia/Source/RivermaxMedia/Private/RivermaxMediaTextureSampleConverter.cpp b/Engine/Plugins/VirtualProduction/Rivermax/RivermaxMedia/Source/RivermaxMedia/Private/RivermaxMediaTextureSampleConverter.cpp
--- a/Engine/Plugins/VirtualProduction/Rivermax/RivermaxMedia/Source/RivermaxMedia/Private/RivermaxMediaTextureSampleConverter.cpp
+++ b/Engine/Plugins/VirtualProduction/Rivermax/RivermaxMedia/Source/RivermaxMedia/Private/RivermaxMediaTextureSampleConverter.cpp
@@ -37,6 +37,37 @@ bool FRivermaxMediaTextureSampleConverter::Convert(FTexture2DRHIRef& InDestinati
 	{
 		break;
 	}
+	case ERivermaxMediaSourePixelFormat::YUV422_10bit:
+	{
+		constexpr uint32 PixelPerGroup = 2;
+		constexpr uint32 BytesPerGroup = 5;
+		const uint32 HorizontalByteCount = (InDestinationTexture->GetDesc().Extent.X / PixelPerGroup) * BytesPerGroup;
+		const int32 BytesPerElement = sizeof(FYUV10Bit422ToRGBACS::FYUV10Bit422LEBuffer);
+		const int32 ElementsPerRow = HorizontalByteCount / BytesPerElement;
+		const int32 ElementCount = ElementsPerRow * InDestinationTexture->GetDesc().Extent.Y;
+
+		FYUV10Bit422ToRGBACS::FPermutationDomain PermutationVector;
+		PermutationVector.Set<FYUV10Bit422ToRGBACS::FSRGBToLinear>(bDoSRGBToLinear);
+
+		FRDGBufferRef InputYUVBuffer = CreateStructuredBuffer(GraphBuilder, TEXT("RivermaxInputBuffer"), BytesPerElement, ElementCount, SamplePtr->GetBuffer(), BytesPerElement * ElementCount);
+		constexpr int32 PixelsPerInput = 8;
+		const FIntPoint ProcessedOutputDimension = { InDestinationTexture->GetDesc().Extent.X / PixelsPerInput,InDestinationTexture->GetDesc().Extent.Y };
+		GroupCount = FComputeShaderUtils::GetGroupCount(ProcessedOutputDimension, FComputeShaderUtils::kGolden2DGroupSize);
+		FGlobalShaderMap* GlobalShaderMap = GetGlobalShaderMap(GMaxRHIFeatureLevel);
+		TShaderMapRef<FYUV10Bit422ToRGBACS> ComputeShader(GlobalShaderMap, PermutationVector);
+		FMatrix YUVToRGBMatrix = SamplePtr->GetYUVToRGBMatrix();
+		FVector YUVOffset(MediaShaders::YUVOffset10bits);
+		FYUV10Bit422ToRGBACS::FParameters* Parameters = ComputeShader->AllocateAndSetParameters(GraphBuilder, InputYUVBuffer, OutputResource, YUVToRGBMatrix, YUVOffset, ElementsPerRow);
+
+
+		FComputeShaderUtils::AddPass(
+			GraphBuilder
+			, RDG_EVENT_NAME("YUV10Bit422ToRGBA")
+			, ComputeShader
+			, Parameters
+			, GroupCount);
+		break;
+	}
 	case ERivermaxMediaSourePixelFormat::RGB_8bit:
 	{
 		constexpr uint32 PixelPerGroup = 1;
diff --git a/Engine/Plugins/VirtualProduction/Rivermax/RivermaxMedia/Source/RivermaxMedia/Public/RivermaxMediaSource.h b/Engine/Plugins/VirtualProduction/Rivermax/RivermaxMedia/Source/RivermaxMedia/Public/RivermaxMediaSource.h
--- a/Engine/Plugins/VirtualProduction/Rivermax/RivermaxMedia/Source/RivermaxMedia/Public/RivermaxMediaSource.h
+++ b/Engine/Plugins/VirtualProduction/Rivermax/RivermaxMedia/Source/RivermaxMedia/Public/RivermaxMediaSource.h
@@ -15,9 +15,10 @@ UENUM()
 enum class ERivermaxMediaSourePixelFormat : uint8
 {
 	YUV422_8bit UMETA(DisplayName = "8bit YUV422"),
+	YUV422_10bit UMETA(DisplayName = "10bit YUV422"),
 	RGB_8bit UMETA(DisplayName = "8bit RGB"),
 	RGB_10bit UMETA(DisplayName = "10bit RGB"),
-	//Todo Add support for 10bit YUV and float16 rgb
+	//Todo Add support for float16 rgb
 };
 
 /**
