- Protecting rivermax shaders from writing out of bound (happens with non regular resolution)
- Incrasing max rivermax input stream stride size to support bigger SRD

diff --git a/Engine/Plugins/VirtualProduction/Rivermax/RivermaxCore/Shaders/Private/RivermaxShaders.usf b/Engine/Plugins/VirtualProduction/Rivermax/RivermaxCore/Shaders/Private/RivermaxShaders.usf
--- a/Engine/Plugins/VirtualProduction/Rivermax/RivermaxCore/Shaders/Private/RivermaxShaders.usf
+++ b/Engine/Plugins/VirtualProduction/Rivermax/RivermaxCore/Shaders/Private/RivermaxShaders.usf
@@ -17,6 +17,7 @@ float OnePixelDeltaY;
 float InputPixelOffsetX;
 float InputPixelOffsetY;
 uint HorizontalElementCount;
+uint VerticalElementCount;
 uint InputTexturePixelsPerThread;
 RWTexture2D<float4> OutTexture;
 
@@ -60,6 +61,12 @@ RWStructuredBuffer<FRGB16fBuffer> OutRGB16fBuffer;
 [numthreads(8,8,1)]
 void RGBToYUV10Bit422(uint3 ThreadId : SV_DispatchThreadID, uint GroupIndex : SV_GroupIndex)
 {
+	// Make sure it's a valid thread
+	if(ThreadId.x >= HorizontalElementCount || ThreadId.y >= VerticalElementCount)
+	{
+		return;
+	}
+
 	uint Index = ThreadId.y * HorizontalElementCount + ThreadId.x;
 
 	float X = (ThreadId.x * InputTexturePixelsPerThread) + InputPixelOffsetX;
@@ -99,6 +106,12 @@ void RGBToYUV10Bit422(uint3 ThreadId : SV_DispatchThreadID, uint GroupIndex : SV
 [numthreads(8,8,1)]
 void RGBToRGB8BitCS(uint3 ThreadId : SV_DispatchThreadID, uint GroupIndex : SV_GroupIndex)
 {
+	// Make sure it's a valid thread
+	if(ThreadId.x >= HorizontalElementCount || ThreadId.y >= VerticalElementCount)
+	{
+		return;
+	}
+
 	uint Index = ThreadId.y * HorizontalElementCount + ThreadId.x;
 
 	float X = (ThreadId.x * InputTexturePixelsPerThread) + InputPixelOffsetX;
@@ -121,6 +134,12 @@ void RGBToRGB8BitCS(uint3 ThreadId : SV_DispatchThreadID, uint GroupIndex : SV_G
 [numthreads(8,8,1)]
 void RGBToRGB10BitCS(uint3 ThreadId : SV_DispatchThreadID, uint GroupIndex : SV_GroupIndex)
 {
+	// Make sure it's a valid thread
+	if(ThreadId.x >= HorizontalElementCount || ThreadId.y >= VerticalElementCount)
+	{
+		return;
+	}
+
 	uint Index = ThreadId.y * HorizontalElementCount + ThreadId.x;
 
 	float X = (ThreadId.x * InputTexturePixelsPerThread) + InputPixelOffsetX;
@@ -177,6 +196,12 @@ void RGBToRGB10BitCS(uint3 ThreadId : SV_DispatchThreadID, uint GroupIndex : SV_
 [numthreads(8,8,1)]
 void RGBToRGB16fBitCS(uint3 ThreadId : SV_DispatchThreadID, uint GroupIndex : SV_GroupIndex)
 {
+	// Make sure it's a valid thread
+	if(ThreadId.x >= HorizontalElementCount || ThreadId.y >= VerticalElementCount)
+	{
+		return;
+	}
+
 	float X = (ThreadId.x * InputTexturePixelsPerThread) + InputPixelOffsetX;
 	float Y = (ThreadId.y + InputPixelOffsetY) * OnePixelDeltaY;
 
@@ -198,6 +223,12 @@ void RGBToRGB16fBitCS(uint3 ThreadId : SV_DispatchThreadID, uint GroupIndex : SV
 [numthreads(8,8,1)]
 void RGB8BitToRGBA8CS(uint3 ThreadId : SV_DispatchThreadID, uint GroupIndex : SV_GroupIndex)
 {
+	// Make sure it's a valid thread
+	if(ThreadId.x >= HorizontalElementCount || ThreadId.y >= VerticalElementCount)
+	{
+		return;
+	}
+
 	// Pixels we treat in shader is pretty much hardcoded to 16 so no parameter has been used
 	const uint OutputTexturePixelsPerThread = 4;
 
@@ -250,6 +281,12 @@ void RGB8BitToRGBA8CS(uint3 ThreadId : SV_DispatchThreadID, uint GroupIndex : SV
 [numthreads(8,8,1)]
 void RGB10BitToRGBACS(uint3 ThreadId : SV_DispatchThreadID, uint GroupIndex : SV_GroupIndex)
 {
+	// Make sure it's a valid thread
+	if(ThreadId.x >= HorizontalElementCount || ThreadId.y >= VerticalElementCount)
+	{
+		return;
+	}
+
 	// Pixels we treat in shader is pretty much hardcoded to 16 so no parameter has been used
 	const uint OutputTexturePixelsPerThread = 16;
 
@@ -387,6 +424,12 @@ void RGB10BitToRGBACS(uint3 ThreadId : SV_DispatchThreadID, uint GroupIndex : SV
 [numthreads(8,8,1)]
 void YUV10Bit422ToRGBACS(uint3 ThreadId : SV_DispatchThreadID, uint GroupIndex : SV_GroupIndex)
 {
+	// Make sure it's a valid thread
+	if(ThreadId.x >= HorizontalElementCount || ThreadId.y >= VerticalElementCount)
+	{
+		return;
+	}
+
 	// Pixels we treat in shader is pretty much hardcoded to 8 so no parameter has been used
 	const uint OutputTexturePixelsPerThread = 8;
 
diff --git a/Engine/Plugins/VirtualProduction/Rivermax/RivermaxCore/Source/RivermaxCore/Private/RivermaxInputStream.cpp b/Engine/Plugins/VirtualProduction/Rivermax/RivermaxCore/Source/RivermaxCore/Private/RivermaxInputStream.cpp
--- a/Engine/Plugins/VirtualProduction/Rivermax/RivermaxCore/Source/RivermaxCore/Private/RivermaxInputStream.cpp
+++ b/Engine/Plugins/VirtualProduction/Rivermax/RivermaxCore/Source/RivermaxCore/Private/RivermaxInputStream.cpp
@@ -151,7 +151,12 @@ namespace UE::RivermaxCore::Private
 					BufferAttributes.attr_flags = BufferAttributeFlags;
 
 					FMemory::Memset(&BufferConfiguration.DataMemory, 0, sizeof(BufferConfiguration.DataMemory));
-					BufferConfiguration.DataMemory.max_size = BufferConfiguration.DataMemory.min_size = BufferConfiguration.PayloadExpectedSize;
+
+					// Todo : Unreal doesn't send Payloads greater than 1280 but other devices could.
+					// While consuming input chunks, we memcpy based on SRD data but SRD data could be greater than 
+					// max_size so our loop should be updated to support that case.
+					BufferConfiguration.DataMemory.min_size = 1500; 
+					BufferConfiguration.DataMemory.max_size = 1500;
 					BufferAttributes.data = &BufferConfiguration.DataMemory;
 
 					FMemory::Memset(&BufferConfiguration.HeaderMemory, 0, sizeof(BufferConfiguration.HeaderMemory));
@@ -346,7 +351,7 @@ namespace UE::RivermaxCore::Private
 							++StreamStats.TotalPacketLossCount;
 							++StreamStats.FramePacketLossCount;
 							
-							UE_LOG(LogRivermax, Warning, TEXT("Lost %uld packets"), LostPackets);
+							UE_LOG(LogRivermax, Warning, TEXT("Lost %llu packets"), LostPackets);
 						}
 
 						StreamData.LastSequenceNumber = Parameter.SequencerNumber;
@@ -361,7 +366,6 @@ namespace UE::RivermaxCore::Private
 
 						if (bCanProcessChunk)
 						{
-							TRACE_CPUPROFILER_EVENT_SCOPE(FRivermaxInputStream::ProcessingChunk)
 							const FRivermaxRTPSampleRowData* HeaderStart = reinterpret_cast<FRivermaxRTPSampleRowData*>(GetRTPHeaderPointer(HeaderPtr));
 
 							uint16 SizeToCopy = ByteSwap((uint16)HeaderStart->SRDLength1);
@@ -390,6 +394,7 @@ namespace UE::RivermaxCore::Private
 							{
 								if (StreamData.ReceivedSize == StreamData.ExpectedSize)
 								{
+									TRACE_BOOKMARK(TEXT("RMAX RX %llu"), Parameter.Timestamp);
 									TRACE_CPUPROFILER_EVENT_SCOPE(FRivermaxInputStream::ProcessingReceivedFrame)
 
 									++StreamStats.FramesReceived;
diff --git a/Engine/Plugins/VirtualProduction/Rivermax/RivermaxCore/Source/RivermaxRendering/Private/RivermaxShaders.cpp b/Engine/Plugins/VirtualProduction/Rivermax/RivermaxCore/Source/RivermaxRendering/Private/RivermaxShaders.cpp
--- a/Engine/Plugins/VirtualProduction/Rivermax/RivermaxCore/Source/RivermaxRendering/Private/RivermaxShaders.cpp
+++ b/Engine/Plugins/VirtualProduction/Rivermax/RivermaxCore/Source/RivermaxRendering/Private/RivermaxShaders.cpp
@@ -13,13 +13,14 @@ namespace UE::RivermaxShaders
 
 	IMPLEMENT_GLOBAL_SHADER(FYUV10Bit422ToRGBACS, "/Plugin/RivermaxCore/Private/RivermaxShaders.usf", "YUV10Bit422ToRGBACS", SF_Compute);
 
-	FYUV10Bit422ToRGBACS::FParameters* FYUV10Bit422ToRGBACS::AllocateAndSetParameters(FRDGBuilder& GraphBuilder, FRDGBufferRef YUVBuffer, FRDGTextureRef OutputTexture, const FMatrix& ColorTransform, const FVector& YUVOffset, int32 BufferElementsPerRow)
+	FYUV10Bit422ToRGBACS::FParameters* FYUV10Bit422ToRGBACS::AllocateAndSetParameters(FRDGBuilder& GraphBuilder, FRDGBufferRef YUVBuffer, FRDGTextureRef OutputTexture, const FMatrix& ColorTransform, const FVector& YUVOffset, int32 BufferElementsPerRow, int32 BufferLineCount)
 	{
 		FYUV10Bit422ToRGBACS::FParameters* Parameters = GraphBuilder.AllocParameters<FYUV10Bit422ToRGBACS::FParameters>();
 
 		Parameters->InputYCbCrBuffer = GraphBuilder.CreateSRV(YUVBuffer);
 		Parameters->ColorTransform = (FMatrix44f)MediaShaders::CombineColorTransformAndOffset(ColorTransform, YUVOffset);
 		Parameters->HorizontalElementCount = BufferElementsPerRow;
+		Parameters->VerticalElementCount = BufferLineCount;
 		Parameters->OutTexture = GraphBuilder.CreateUAV(OutputTexture, ERDGUnorderedAccessViewFlags::None);
 
 		return Parameters;
@@ -49,6 +50,7 @@ namespace UE::RivermaxShaders
 
 		// Output size will be based on RGB to YUV encoding ( divide by 2 ) and pixels encoded per element ( 4 pixels )
 		Parameters->HorizontalElementCount = OutputSize.X;
+		Parameters->VerticalElementCount = OutputSize.Y;
 
 		// Each thread will read many RGB pixels to convert to YUV
 		Parameters->InputTexturePixelsPerThread = InputTextureSizeX / OutputSize.X;
@@ -63,12 +65,13 @@ namespace UE::RivermaxShaders
 
 	IMPLEMENT_GLOBAL_SHADER(FRGB8BitToRGBA8CS, "/Plugin/RivermaxCore/Private/RivermaxShaders.usf", "RGB8BitToRGBA8CS", SF_Compute);
 
-	FRGB8BitToRGBA8CS::FParameters* FRGB8BitToRGBA8CS::AllocateAndSetParameters(FRDGBuilder& GraphBuilder, FRDGBufferRef RGBBuffer, FRDGTextureRef OutputTexture, int32 BufferElementsPerRow)
+	FRGB8BitToRGBA8CS::FParameters* FRGB8BitToRGBA8CS::AllocateAndSetParameters(FRDGBuilder& GraphBuilder, FRDGBufferRef RGBBuffer, FRDGTextureRef OutputTexture, int32 BufferElementsPerRow, int32 BufferLineCount)
 	{
 		FRGB8BitToRGBA8CS::FParameters* Parameters = GraphBuilder.AllocParameters<FRGB8BitToRGBA8CS::FParameters>();
 
 		Parameters->InputRGB8Buffer = GraphBuilder.CreateSRV(RGBBuffer);
 		Parameters->HorizontalElementCount = BufferElementsPerRow;
+		Parameters->VerticalElementCount = BufferLineCount;
 		Parameters->OutTexture = GraphBuilder.CreateUAV(OutputTexture, ERDGUnorderedAccessViewFlags::None);
 
 		return Parameters;
@@ -95,6 +98,7 @@ namespace UE::RivermaxShaders
 		Parameters->InputPixelOffsetY = SourceViewRect.Min.Y;
 
 		Parameters->HorizontalElementCount = OutputSize.X;
+		Parameters->VerticalElementCount = OutputSize.Y;
 
 		// Each thread will read many RGBA pixels to pack into RGB buffer
 		Parameters->InputTexturePixelsPerThread = InputTextureSizeX / OutputSize.X;
@@ -110,12 +114,13 @@ namespace UE::RivermaxShaders
 	IMPLEMENT_GLOBAL_SHADER(FRGB10BitToRGBA10CS, "/Plugin/RivermaxCore/Private/RivermaxShaders.usf", "RGB10BitToRGBACS", SF_Compute);
 
 
-	FRGB10BitToRGBA10CS::FParameters* FRGB10BitToRGBA10CS::AllocateAndSetParameters(FRDGBuilder& GraphBuilder, FRDGBufferRef RGBBuffer, FRDGTextureRef OutputTexture, int32 BufferElementsPerRow)
+	FRGB10BitToRGBA10CS::FParameters* FRGB10BitToRGBA10CS::AllocateAndSetParameters(FRDGBuilder& GraphBuilder, FRDGBufferRef RGBBuffer, FRDGTextureRef OutputTexture, int32 BufferElementsPerRow, int32 BufferLineCount)
 	{
 		FRGB10BitToRGBA10CS::FParameters* Parameters = GraphBuilder.AllocParameters<FRGB10BitToRGBA10CS::FParameters>();
 
 		Parameters->InputBuffer = GraphBuilder.CreateSRV(RGBBuffer);
 		Parameters->HorizontalElementCount = BufferElementsPerRow;
+		Parameters->VerticalElementCount = BufferLineCount;
 		Parameters->OutTexture = GraphBuilder.CreateUAV(OutputTexture, ERDGUnorderedAccessViewFlags::None);
 
 		return Parameters;
@@ -143,6 +148,7 @@ namespace UE::RivermaxShaders
 		Parameters->InputPixelOffsetY = SourceViewRect.Min.Y;
 
 		Parameters->HorizontalElementCount = OutputSize.X;
+		Parameters->VerticalElementCount = OutputSize.Y;
 
 		// Each thread will read many RGBA pixels to pack into RGB buffer
 		Parameters->InputTexturePixelsPerThread = InputTextureSizeX / OutputSize.X;
@@ -174,6 +180,7 @@ namespace UE::RivermaxShaders
 		Parameters->InputPixelOffsetY = SourceViewRect.Min.Y;
 
 		Parameters->HorizontalElementCount = OutputSize.X;
+		Parameters->VerticalElementCount = OutputSize.Y;
 
 		// Each thread will read many RGBA pixels to pack into RGB buffer
 		Parameters->InputTexturePixelsPerThread = InputTextureSizeX / OutputSize.X;
diff --git a/Engine/Plugins/VirtualProduction/Rivermax/RivermaxCore/Source/RivermaxRendering/Public/RivermaxShaders.h b/Engine/Plugins/VirtualProduction/Rivermax/RivermaxCore/Source/RivermaxRendering/Public/RivermaxShaders.h
--- a/Engine/Plugins/VirtualProduction/Rivermax/RivermaxCore/Source/RivermaxRendering/Public/RivermaxShaders.h
+++ b/Engine/Plugins/VirtualProduction/Rivermax/RivermaxCore/Source/RivermaxRendering/Public/RivermaxShaders.h
@@ -40,6 +40,7 @@ namespace UE::RivermaxShaders
 			SHADER_PARAMETER_RDG_BUFFER_SRV(StructuredBuffer<FYUV10Bit422LEBuffer>, InputYCbCrBuffer)
 			SHADER_PARAMETER(FMatrix44f, ColorTransform)	
 			SHADER_PARAMETER(uint32, HorizontalElementCount)
+			SHADER_PARAMETER(uint32, VerticalElementCount)
 			SHADER_PARAMETER_RDG_TEXTURE_UAV(RWTexture2D, OutTexture)
 		END_SHADER_PARAMETER_STRUCT()
 
@@ -52,7 +53,7 @@ namespace UE::RivermaxShaders
 		}
 
 		/** Allocates and setup shader parameter in the incoming graph builder */
-		FYUV10Bit422ToRGBACS::FParameters* AllocateAndSetParameters(FRDGBuilder& GraphBuilder, FRDGBufferRef YUVBuffer, FRDGTextureRef OutputTexture, const FMatrix& ColorTransform, const FVector& YUVOffset, int32 BufferElementsPerRow);
+		FYUV10Bit422ToRGBACS::FParameters* AllocateAndSetParameters(FRDGBuilder& GraphBuilder, FRDGBufferRef YUVBuffer, FRDGTextureRef OutputTexture, const FMatrix& ColorTransform, const FVector& YUVOffset, int32 BufferElementsPerRow, int32 BufferLineCount);
 	};
 
 	/**
@@ -83,6 +84,7 @@ namespace UE::RivermaxShaders
 			SHADER_PARAMETER(float, InputPixelOffsetX)
 			SHADER_PARAMETER(float, InputPixelOffsetY)
 			SHADER_PARAMETER(uint32, HorizontalElementCount)
+			SHADER_PARAMETER(uint32, VerticalElementCount)
 			SHADER_PARAMETER(uint32, InputTexturePixelsPerThread)
 			SHADER_PARAMETER_RDG_BUFFER_UAV(RWStructuredBuffer<FYUV10Bit422LEBuffer>, OutYCbCrBuffer)
 		END_SHADER_PARAMETER_STRUCT()
@@ -124,6 +126,7 @@ namespace UE::RivermaxShaders
 		BEGIN_SHADER_PARAMETER_STRUCT(FParameters, )
 			SHADER_PARAMETER_RDG_BUFFER_SRV(StructuredBuffer<FRGB8BitBuffer>, InputRGB8Buffer)
 			SHADER_PARAMETER(uint32, HorizontalElementCount)
+			SHADER_PARAMETER(uint32, VerticalElementCount)
 			SHADER_PARAMETER_RDG_TEXTURE_UAV(RWTexture2D, OutTexture)
 		END_SHADER_PARAMETER_STRUCT()
 
@@ -136,7 +139,7 @@ namespace UE::RivermaxShaders
 		}
 
 		/** Allocates and setup shader parameter in the incoming graph builder */
-		FRGB8BitToRGBA8CS::FParameters* AllocateAndSetParameters(FRDGBuilder& GraphBuilder, FRDGBufferRef RGBBuffer, FRDGTextureRef OutputTexture, int32 BufferElementsPerRow);
+		FRGB8BitToRGBA8CS::FParameters* AllocateAndSetParameters(FRDGBuilder& GraphBuilder, FRDGBufferRef RGBBuffer, FRDGTextureRef OutputTexture, int32 BufferElementsPerRow, int32 BufferLineCount);
 	};
 
 	/**
@@ -167,6 +170,7 @@ namespace UE::RivermaxShaders
 			SHADER_PARAMETER(float, InputPixelOffsetX)
 			SHADER_PARAMETER(float, InputPixelOffsetY)
 			SHADER_PARAMETER(uint32, HorizontalElementCount)
+			SHADER_PARAMETER(uint32, VerticalElementCount)
 			SHADER_PARAMETER(uint32, InputTexturePixelsPerThread)
 			SHADER_PARAMETER_RDG_BUFFER_UAV(RWStructuredBuffer<FRGB8BitBuffer>, OutRGB8Buffer)
 		END_SHADER_PARAMETER_STRUCT()
@@ -210,6 +214,7 @@ namespace UE::RivermaxShaders
 			SHADER_PARAMETER(float, InputPixelOffsetX)
 			SHADER_PARAMETER(float, InputPixelOffsetY)
 			SHADER_PARAMETER(uint32, HorizontalElementCount)
+			SHADER_PARAMETER(uint32, VerticalElementCount)
 			SHADER_PARAMETER(uint32, InputTexturePixelsPerThread)
 			SHADER_PARAMETER_RDG_BUFFER_UAV(RWStructuredBuffer<FRGB10BitBuffer>, OutRGB10Buffer)
 			END_SHADER_PARAMETER_STRUCT()
@@ -250,6 +255,7 @@ namespace UE::RivermaxShaders
 		BEGIN_SHADER_PARAMETER_STRUCT(FParameters, )
 			SHADER_PARAMETER_RDG_BUFFER_SRV(StructuredBuffer<FRGB10BitBuffer>, InputBuffer)
 			SHADER_PARAMETER(uint32, HorizontalElementCount)
+			SHADER_PARAMETER(uint32, VerticalElementCount)
 			SHADER_PARAMETER_RDG_TEXTURE_UAV(RWTexture2D, OutTexture)
 		END_SHADER_PARAMETER_STRUCT()
 
@@ -262,7 +268,7 @@ namespace UE::RivermaxShaders
 		}
 
 		/** Allocates and setup shader parameter in the incoming graph builder */
-		FRGB10BitToRGBA10CS::FParameters* AllocateAndSetParameters(FRDGBuilder& GraphBuilder, FRDGBufferRef RGBBuffer, FRDGTextureRef OutputTexture, int32 BufferElementsPerRow);
+		FRGB10BitToRGBA10CS::FParameters* AllocateAndSetParameters(FRDGBuilder& GraphBuilder, FRDGBufferRef RGBBuffer, FRDGTextureRef OutputTexture, int32 BufferElementsPerRow, int32 BufferLineCount);
 	};
 
 	/**
@@ -292,6 +298,7 @@ namespace UE::RivermaxShaders
 			SHADER_PARAMETER(float, InputPixelOffsetX)
 			SHADER_PARAMETER(float, InputPixelOffsetY)
 			SHADER_PARAMETER(uint32, HorizontalElementCount)
+			SHADER_PARAMETER(uint32, VerticalElementCount)
 			SHADER_PARAMETER(uint32, InputTexturePixelsPerThread)
 			SHADER_PARAMETER_RDG_BUFFER_UAV(RWStructuredBuffer<FRGB16fBuffer>, OutRGB16fBuffer)
 		END_SHADER_PARAMETER_STRUCT()
diff --git a/Engine/Plugins/VirtualProduction/Rivermax/RivermaxMedia/Source/RivermaxMedia/Private/RivermaxMediaTextureSampleConverter.cpp b/Engine/Plugins/VirtualProduction/Rivermax/RivermaxMedia/Source/RivermaxMedia/Private/RivermaxMediaTextureSampleConverter.cpp
--- a/Engine/Plugins/VirtualProduction/Rivermax/RivermaxMedia/Source/RivermaxMedia/Private/RivermaxMediaTextureSampleConverter.cpp
+++ b/Engine/Plugins/VirtualProduction/Rivermax/RivermaxMedia/Source/RivermaxMedia/Private/RivermaxMediaTextureSampleConverter.cpp
@@ -51,13 +51,13 @@ bool FRivermaxMediaTextureSampleConverter::Convert(FTexture2DRHIRef& InDestinati
 
 		FRDGBufferRef InputYUVBuffer = CreateStructuredBuffer(GraphBuilder, TEXT("RivermaxInputBuffer"), BytesPerElement, ElementCount, SamplePtr->GetBuffer(), BytesPerElement * ElementCount);
 		constexpr int32 PixelsPerInput = 8;
-		const FIntPoint ProcessedOutputDimension = { InDestinationTexture->GetDesc().Extent.X / PixelsPerInput,InDestinationTexture->GetDesc().Extent.Y };
+		const FIntPoint ProcessedOutputDimension = { InDestinationTexture->GetDesc().Extent.X / PixelsPerInput, InDestinationTexture->GetDesc().Extent.Y };
 		GroupCount = FComputeShaderUtils::GetGroupCount(ProcessedOutputDimension, FComputeShaderUtils::kGolden2DGroupSize);
 		FGlobalShaderMap* GlobalShaderMap = GetGlobalShaderMap(GMaxRHIFeatureLevel);
 		TShaderMapRef<FYUV10Bit422ToRGBACS> ComputeShader(GlobalShaderMap, PermutationVector);
 		FMatrix YUVToRGBMatrix = SamplePtr->GetYUVToRGBMatrix();
 		FVector YUVOffset(MediaShaders::YUVOffset10bits);
-		FYUV10Bit422ToRGBACS::FParameters* Parameters = ComputeShader->AllocateAndSetParameters(GraphBuilder, InputYUVBuffer, OutputResource, YUVToRGBMatrix, YUVOffset, ElementsPerRow);
+		FYUV10Bit422ToRGBACS::FParameters* Parameters = ComputeShader->AllocateAndSetParameters(GraphBuilder, InputYUVBuffer, OutputResource, YUVToRGBMatrix, YUVOffset, ElementsPerRow, ProcessedOutputDimension.Y);
 
 
 		FComputeShaderUtils::AddPass(
@@ -86,7 +86,7 @@ bool FRivermaxMediaTextureSampleConverter::Convert(FTexture2DRHIRef& InDestinati
 		GroupCount = FComputeShaderUtils::GetGroupCount(ProcessedOutputDimension, FComputeShaderUtils::kGolden2DGroupSize);
 		FGlobalShaderMap* GlobalShaderMap = GetGlobalShaderMap(GMaxRHIFeatureLevel);
 		TShaderMapRef<FRGB8BitToRGBA8CS> ComputeShader(GlobalShaderMap, PermutationVector);
-		FRGB8BitToRGBA8CS::FParameters* Parameters = ComputeShader->AllocateAndSetParameters(GraphBuilder, InputRGGBuffer, OutputResource, ElementsPerRow);
+		FRGB8BitToRGBA8CS::FParameters* Parameters = ComputeShader->AllocateAndSetParameters(GraphBuilder, InputRGGBuffer, OutputResource, ElementsPerRow, ProcessedOutputDimension.Y);
 
 
 		FComputeShaderUtils::AddPass(
@@ -115,7 +115,7 @@ bool FRivermaxMediaTextureSampleConverter::Convert(FTexture2DRHIRef& InDestinati
 		GroupCount = FComputeShaderUtils::GetGroupCount(ProcessedOutputDimension, FComputeShaderUtils::kGolden2DGroupSize);
 		FGlobalShaderMap* GlobalShaderMap = GetGlobalShaderMap(GMaxRHIFeatureLevel);
 		TShaderMapRef<FRGB10BitToRGBA10CS> ComputeShader(GlobalShaderMap, PermutationVector);
-		FRGB10BitToRGBA10CS::FParameters* Parameters = ComputeShader->AllocateAndSetParameters(GraphBuilder, InputRGGBuffer, OutputResource, ElementsPerRow);
+		FRGB10BitToRGBA10CS::FParameters* Parameters = ComputeShader->AllocateAndSetParameters(GraphBuilder, InputRGGBuffer, OutputResource, ElementsPerRow, ProcessedOutputDimension.Y);
 
 
 		FComputeShaderUtils::AddPass(
