- Adding support for rivermax clock selection
- Adding support for rivermax device picker
- Adding support for wildcards in rivermax address

diff --git a/Engine/Plugins/VirtualProduction/Rivermax/RivermaxCore/RivermaxCore.uplugin b/Engine/Plugins/VirtualProduction/Rivermax/RivermaxCore/RivermaxCore.uplugin
--- a/Engine/Plugins/VirtualProduction/Rivermax/RivermaxCore/RivermaxCore.uplugin
+++ b/Engine/Plugins/VirtualProduction/Rivermax/RivermaxCore/RivermaxCore.uplugin
@@ -28,6 +28,12 @@
 			"LoadingPhase": "Default",
 			"PlatformAllowList": [ "Win64" ]
 		},
+		{
+			"Name": "RivermaxEditor",
+			"Type": "Editor",
+			"LoadingPhase": "Default",
+			"PlatformAllowList": [ "Win64" ]
+		},
 		{
 			"Name": "RivermaxRendering",
 			"Type": "Runtime",
diff --git a/Engine/Plugins/VirtualProduction/Rivermax/RivermaxCore/Source/RivermaxCore/Private/RivermaxDeviceFinder.cpp b/Engine/Plugins/VirtualProduction/Rivermax/RivermaxCore/Source/RivermaxCore/Private/RivermaxDeviceFinder.cpp
new file mode 100644
--- /dev/null
+++ b/Engine/Plugins/VirtualProduction/Rivermax/RivermaxCore/Source/RivermaxCore/Private/RivermaxDeviceFinder.cpp
@@ -0,0 +1,194 @@
+// Copyright Epic Games, Inc. All Rights Reserved.
+
+#include "RivermaxDeviceFinder.h"
+
+#include "Misc/LazySingleton.h"
+#include "IRivermaxCoreModule.h"
+#include "RivermaxLog.h"
+
+
+#if PLATFORM_WINDOWS
+#include "Windows/AllowWindowsPlatformTypes.h"
+#include "Windows/PreWindowsApi.h"
+
+#include <winsock2.h>
+#include <Iphlpapi.h>
+#include <ws2tcpip.h>
+
+#include "Windows/PostWindowsApi.h"
+#include "Windows/HideWindowsPlatformTypes.h"
+#endif
+
+
+namespace UE::RivermaxCore::Private
+{
+	FRivermaxDeviceFinder::FRivermaxDeviceFinder()
+	{
+		FindDevices();
+	}
+
+	bool FRivermaxDeviceFinder::IsValidIP(const FString& SourceIP) const
+	{
+		const uint32 IPHash = GetTypeHash(SourceIP);
+		if (FQueriedIPInfo* FoundQuery = QueriesCache.Find(IPHash))
+		{
+			return FoundQuery->bIsValidIP;
+		}
+
+		TArray<FString> SourceTokens;
+		bool bIsValid = SourceIP.ParseIntoArray(SourceTokens, TEXT("."), false /*CullEmpty*/) == 4;
+		if (bIsValid)
+		{
+			for (const FString& Token : SourceTokens)
+			{
+				if (FCString::IsNumeric(*Token) == false && Token.Equals(TEXT("*")) == false)
+				{
+					bIsValid = false;
+					break;
+				}
+			}
+		}
+
+		FQueriedIPInfo& FoundQuery = QueriesCache.FindOrAdd(IPHash);
+		FoundQuery.bIsValidIP = bIsValid;
+		FoundQuery.SourceIP = SourceIP;
+		FoundQuery.Tokens = MoveTemp(SourceTokens);
+
+		return bIsValid;
+	}
+
+	bool FRivermaxDeviceFinder::ResolveIP(const FString& SourceIP, FString& OutDeviceIP) const
+	{
+		if (IsValidIP(SourceIP))
+		{
+			const uint32 IPHash = GetTypeHash(SourceIP);
+			FQueriedIPInfo& FoundQuery = QueriesCache[IPHash];
+			if (FoundQuery.ResolvedDeviceIPHash.IsSet())
+			{
+				// Device cache should have this entry if the hash has been set in the IP cache
+				OutDeviceIP = DevicesCache[FoundQuery.ResolvedDeviceIPHash.GetValue()].DeviceIP;
+				return true;
+			}
+
+			for (const FRivermaxDeviceInfo& Device : Devices)
+			{
+				const uint32 DeviceIPHash = GetTypeHash(Device.InterfaceAddress);
+				if (DevicesCache.Contains(DeviceIPHash) == false)
+				{
+					FDeviceInfo& NewDevice = DevicesCache.FindOrAdd(DeviceIPHash);
+					NewDevice.DeviceIP = Device.InterfaceAddress;
+					check(FIPv4Address::Parse(NewDevice.DeviceIP, NewDevice.Address));
+					NewDevice.AddressTokens[0] = NewDevice.Address.A;
+					NewDevice.AddressTokens[1] = NewDevice.Address.B;
+					NewDevice.AddressTokens[2] = NewDevice.Address.C;
+					NewDevice.AddressTokens[3] = NewDevice.Address.D;
+				}
+
+				const FDeviceInfo& CachedDevice = DevicesCache[DeviceIPHash];
+
+				bool bMatchingDeviceFound = true;
+				for (int32 TokenIndex = 0; TokenIndex < 4; ++TokenIndex)
+				{
+					if (FoundQuery.Tokens[TokenIndex].Equals(TEXT("*")) == false)
+					{
+						// If not a wildcard, token must match by value
+						const uint8 TokenValue = FCString::Atoi(*FoundQuery.Tokens[TokenIndex]);
+						const uint8 DeviceTokenValue = CachedDevice.AddressTokens[TokenIndex];
+						if (TokenValue != DeviceTokenValue)
+						{
+							bMatchingDeviceFound = false;
+							break;
+						}
+					}
+				}
+
+				if (bMatchingDeviceFound)
+				{
+					FoundQuery.ResolvedDeviceIPHash = DeviceIPHash;
+					OutDeviceIP = CachedDevice.DeviceIP;
+					return true;
+				}
+			}
+		}
+
+		return false;
+	}
+
+	TConstArrayView<FRivermaxDeviceInfo> FRivermaxDeviceFinder::GetDevices() const
+	{
+		return Devices;
+	}
+
+	void FRivermaxDeviceFinder::FindDevices()
+	{
+		constexpr ULONG Flags = GAA_FLAG_INCLUDE_PREFIX | GAA_FLAG_SKIP_MULTICAST | GAA_FLAG_SKIP_ANYCAST | GAA_FLAG_SKIP_DNS_SERVER | GAA_FLAG_SKIP_FRIENDLY_NAME;
+		constexpr ULONG Family = AF_INET;
+
+		// Based on SocketBSD implementation
+		// determine the required size of the address list buffer
+		ULONG Size = 0;
+		ULONG Result = GetAdaptersAddresses(Family, Flags, NULL, NULL, &Size);
+
+		if (Result != ERROR_BUFFER_OVERFLOW)
+		{
+			return;
+		}
+
+		PIP_ADAPTER_ADDRESSES AdapterAddresses = (PIP_ADAPTER_ADDRESSES)FMemory::Malloc(Size);
+
+		// get the actual list of adapters
+		Result = GetAdaptersAddresses(Family, Flags, NULL, AdapterAddresses, &Size);
+
+		if (Result == ERROR_SUCCESS)
+		{
+			// extract the list of physical addresses from each adapter
+			for (PIP_ADAPTER_ADDRESSES AdapterAddress = AdapterAddresses; AdapterAddress != NULL; AdapterAddress = AdapterAddress->Next)
+			{
+				if (AdapterAddress->OperStatus == IfOperStatusUp)
+				{
+					const auto AdapterFilterFunc = [](const FString& Description) -> bool
+					{
+						// We could also look for Mellanox but it would also include management port which we would need to discard
+						if (Description.Contains(TEXT("ConnectX")))
+						{
+							return true;
+						}
+
+						return false;
+					};
+
+					const FString Description = StringCast<TCHAR>(AdapterAddress->Description).Get();
+					if (AdapterFilterFunc(Description))
+					{
+						for (PIP_ADAPTER_UNICAST_ADDRESS UnicastAddress = AdapterAddress->FirstUnicastAddress; UnicastAddress != NULL; UnicastAddress = UnicastAddress->Next)
+						{
+							sockaddr_storage* RawAddress = reinterpret_cast<sockaddr_storage*>(UnicastAddress->Address.lpSockaddr);
+
+							// Verify if it's IPV4 or 6
+							if (RawAddress->ss_family == AF_INET)
+							{
+								const sockaddr* AdapterUniAddress = reinterpret_cast<const sockaddr*>(RawAddress);
+								char IPString[NI_MAXHOST];
+								if (getnameinfo(AdapterUniAddress, sizeof(sockaddr_in), IPString, NI_MAXHOST, nullptr, 0, NI_NUMERICHOST) == 0)
+								{
+									FRivermaxDeviceInfo NewDevice;
+									NewDevice.Description = Description;
+									NewDevice.InterfaceAddress = StringCast<TCHAR>(IPString).Get();
+
+									UE_LOG(LogRivermax, Log, TEXT("Found adapter: Name: '%s', IP: '%s'"), *NewDevice.Description, *NewDevice.InterfaceAddress);
+									Devices.Add(MoveTemp(NewDevice));
+								}
+							}
+						}
+					}
+				}
+			}
+		}
+
+		FMemory::Free(AdapterAddresses);
+	}
+
+
+}
+
+
diff --git a/Engine/Plugins/VirtualProduction/Rivermax/RivermaxCore/Source/RivermaxCore/Private/RivermaxDeviceFinder.h b/Engine/Plugins/VirtualProduction/Rivermax/RivermaxCore/Source/RivermaxCore/Private/RivermaxDeviceFinder.h
new file mode 100644
--- /dev/null
+++ b/Engine/Plugins/VirtualProduction/Rivermax/RivermaxCore/Source/RivermaxCore/Private/RivermaxDeviceFinder.h
@@ -0,0 +1,72 @@
+// Copyright Epic Games, Inc. All Rights Reserved.
+
+#pragma once
+
+#include "CoreMinimal.h"
+
+#include "Interfaces/IPv4/IPv4Address.h"
+#include "IRivermaxManager.h"
+
+
+namespace UE::RivermaxCore::Private
+{
+	class FRivermaxDeviceFinder
+	{
+	public:
+		FRivermaxDeviceFinder();
+
+		/** Tests whether provided string is a valid IP, containing 4 numeric values or wildcard (*) separated by dots. */
+		bool IsValidIP(const FString& SourceIP) const;
+
+		/** Resolves IP to a found Rivermax device interface. */
+		bool ResolveIP(const FString& SourceIP, FString& OutDeviceIP) const;
+
+		/** Returns list of found devices that can be used */
+		TConstArrayView<FRivermaxDeviceInfo> GetDevices() const;
+
+	private:
+
+		/** Find all interfaces available */
+		void FindDevices();
+
+	private:
+
+		struct FQueriedIPInfo
+		{
+			/** Whether or not this IP was found as valid */
+			bool bIsValidIP = false;
+
+			/** String representation of the IP */
+			FString SourceIP;
+
+			/** String representation of IP component separated by '.' */
+			TArray<FString> Tokens;
+
+			/** Hash of the matching found resolved device IP used to quickly find association */
+			TOptional<uint32> ResolvedDeviceIPHash;
+		};
+
+		struct FDeviceInfo
+		{
+			/** String representation of the IP */
+			FString DeviceIP;
+
+			/** IPv4 representation of the IP with numerical tokens */
+			FIPv4Address Address;
+
+			/** Array representation of IP in numerical tokens to iterate over */
+			uint8 AddressTokens[4];
+		};
+
+		/** Cache of IPs that were queried to see which were valid / resolved */
+		mutable TMap<uint32, FQueriedIPInfo> QueriesCache;
+
+		/** Cache of found Rivermax devices with useful representation of its IP */
+		mutable TMap<uint32, FDeviceInfo> DevicesCache;
+		
+		/** List of found Rivermax devices */
+		TArray<FRivermaxDeviceInfo> Devices;
+	};
+}
+
+
diff --git a/Engine/Plugins/VirtualProduction/Rivermax/RivermaxCore/Source/RivermaxCore/Private/RivermaxManager.cpp b/Engine/Plugins/VirtualProduction/Rivermax/RivermaxCore/Source/RivermaxCore/Private/RivermaxManager.cpp
--- a/Engine/Plugins/VirtualProduction/Rivermax/RivermaxCore/Source/RivermaxCore/Private/RivermaxManager.cpp
+++ b/Engine/Plugins/VirtualProduction/Rivermax/RivermaxCore/Source/RivermaxCore/Private/RivermaxManager.cpp
@@ -2,57 +2,50 @@
 
 #include "RivermaxManager.h"
 
+#include "Interfaces/IPv4/IPv4Endpoint.h"
+#include "IRivermaxCoreModule.h"
 #include "Misc/App.h"
 #include "Misc/Paths.h"
 #include "Misc/CoreDelegates.h"
+#include "RivermaxDeviceFinder.h"
 #include "RivermaxHeader.h"
 #include "RivermaxLog.h"
 
 
+
+#if PLATFORM_WINDOWS
+#include "Windows/AllowWindowsPlatformTypes.h"
+#include "Windows/PreWindowsApi.h"
+
+#include <ws2tcpip.h>
+
+#include "Windows/PostWindowsApi.h"
+#include "Windows/HideWindowsPlatformTypes.h"
+#endif
+
+
+
 namespace UE::RivermaxCore::Private
 {
-	FRivermaxManager::FRivermaxManager()
+	static uint64 RivermaxTimeHandler(void* Context)
 	{
-		const auto InitializeRivermaxLibFunc = [this]()
+		IRivermaxCoreModule* RivermaxModule = FModuleManager::GetModulePtr<IRivermaxCoreModule>("RivermaxCore");
+		if (RivermaxModule)
 		{
-			rmax_init_config Config;
-			memset(&Config, 0, sizeof(Config));
-			//Config.flags |= RIVERMAX_HANDLE_SIGNAL; //Verify if desired
+			return RivermaxModule->GetRivermaxManager()->GetTime();
+		}
 
-			rmax_status_t Status = rmax_init(&Config);
-			if (Status == RMAX_OK)
-			{
-				uint32 Major = 0;
-				uint32 Minor = 0;
-				uint32 Release = 0;
-				uint32 Build = 0;
-				Status = rmax_get_version(&Major, &Minor, &Release, &Build);
-				if (Status == RMAX_OK)
-				{
-					bIsInitialized = true;
-					UE_LOG(LogRivermax, Log, TEXT("Rivermax library version %d.%d.%d.%d succesfully initialized"), Major, Minor, Release, Build);
-				}
-				else
-				{
-					UE_LOG(LogRivermax, Log, TEXT("Failed to retrieve Rivermax library version. Status: %d"), Status);
-				}
-			}
-			else if (Status == RMAX_ERR_LICENSE_ISSUE)
-			{
-				UE_LOG(LogRivermax, Log, TEXT("Rivermax License could not be found. Have you configured RIVERMAX_LICENSE_PATH environment variable?"));
-			}
-			
-			if(bIsInitialized == false)
-			{
-				UE_LOG(LogRivermax, Log, TEXT("Rivermax library failed to initialize. Status: %d"), Status);
-			}
-		};
+		return 0;
+	}
 
+	FRivermaxManager::FRivermaxManager()
+		: DeviceFinder(MakeUnique<FRivermaxDeviceFinder>())
+	{
 		const bool bIsLibraryLoaded = LoadRivermaxLibrary();
 		if (bIsLibraryLoaded && FApp::CanEverRender())
 		{
 			//Postpone initialization after all modules have been loaded
-			FCoreDelegates::OnAllModuleLoadingPhasesComplete.AddLambda(InitializeRivermaxLibFunc);
+			FCoreDelegates::OnAllModuleLoadingPhasesComplete.AddRaw(this, &FRivermaxManager::InitializeLibrary);
 		}
 		else
 		{
@@ -104,6 +97,114 @@ namespace UE::RivermaxCore::Private
 		return bIsLibraryValid;
 	}
 
+	uint64 FRivermaxManager::GetTime() const
+	{
+		uint64 Time = 0;
+
+		if (IsInitialized() && GetTimeSource() == ERivermaxTimeSource::PTP)
+		{
+			//Rivermax is usable and configured to read PTP
+			if (rmax_get_time(RMAX_CLOCK_PTP, &Time) != RMAX_OK)
+			{
+				UE_LOG(LogRivermax, Warning, TEXT("PTP time is the time source but was unavailable"));
+			}
+		}
+		else
+		{
+			// PTP was not available so we fall back on system time
+			Time = uint64(FPlatformTime::Seconds() * 1E9);
+		}
+
+		return Time;
+	}
+
+	ERivermaxTimeSource FRivermaxManager::GetTimeSource() const
+	{
+		return TimeSource;
+	}
+
+	void FRivermaxManager::InitializeLibrary()
+	{
+		rmax_init_config Config;
+		memset(&Config, 0, sizeof(Config));
+		Config.flags |= RIVERMAX_ENABLE_CLOCK_CONFIGURATION;
+			
+		const URivermaxSettings* Settings = GetDefault<URivermaxSettings>();
+
+		bool bConfiguredSuccessfully = false;
+		if (TimeSource == ERivermaxTimeSource::PTP)
+		{
+			Config.clock_configurations.clock_type = rmax_clock_types::RIVERMAX_PTP_CLOCK;
+
+			FString PTPAddress;
+			if (DeviceFinder->ResolveIP(Settings->PTPInterfaceAddress, PTPAddress))
+			{
+				if (inet_pton(AF_INET, StringCast<ANSICHAR>(*PTPAddress).Get(), &Config.clock_configurations.clock_u.rmax_ptp_clock.device_ip_addr))
+				{
+					UE_LOG(LogRivermax, Display, TEXT("Initialize Rivermax clock as PTP using device interface %s"), *PTPAddress);
+					TimeSource = ERivermaxTimeSource::PTP;
+					bConfiguredSuccessfully = true;
+				}
+			}
+
+			if (bConfiguredSuccessfully == false)
+			{
+				UE_LOG(LogRivermax, Warning, TEXT("Could not configure Rivermax to use PTP using interface '%s'. Falling back to platform time."), *Settings->PTPInterfaceAddress);
+			}
+		}
+
+		if(bConfiguredSuccessfully == false)
+		{
+			// Todo: Need to update clock fetcher to align with UTC. Refer to SDK
+			Config.clock_configurations.clock_type = rmax_clock_types::RIVERMAX_USER_CLOCK_HANDLER;
+			Config.clock_configurations.clock_u.rmax_user_clock_handler.clock_handler = RivermaxTimeHandler;
+			Config.clock_configurations.clock_u.rmax_user_clock_handler.ctx = nullptr;
+			TimeSource = ERivermaxTimeSource::Platform;
+		}
+
+		rmax_status_t Status = rmax_init(&Config);
+		if (Status == RMAX_OK)
+		{
+			uint32 Major = 0;
+			uint32 Minor = 0;
+			uint32 Release = 0;
+			uint32 Build = 0;
+			Status = rmax_get_version(&Major, &Minor, &Release, &Build);
+			if (Status == RMAX_OK)
+			{
+				bIsInitialized = true;
+				UE_LOG(LogRivermax, Log, TEXT("Rivermax library version %d.%d.%d.%d succesfully initialized"), Major, Minor, Release, Build);
+			}
+			else
+			{
+				UE_LOG(LogRivermax, Log, TEXT("Failed to retrieve Rivermax library version. Status: %d"), Status);
+			}
+		}
+		else if (Status == RMAX_ERR_LICENSE_ISSUE)
+		{
+			UE_LOG(LogRivermax, Log, TEXT("Rivermax License could not be found. Have you configured RIVERMAX_LICENSE_PATH environment variable?"));
+		}
+
+		if (bIsInitialized == false)
+		{
+			UE_LOG(LogRivermax, Log, TEXT("Rivermax library failed to initialize. Status: %d"), Status);
+		}
+	}
+
+	TConstArrayView<UE::RivermaxCore::FRivermaxDeviceInfo> FRivermaxManager::GetDevices() const
+	{
+		return DeviceFinder->GetDevices();
+	}
+
+	bool FRivermaxManager::GetMatchingDevice(const FString& InSourceIP, FString& OutDeviceIP) const
+	{
+		return DeviceFinder->ResolveIP(InSourceIP, OutDeviceIP);
+	}
+
+	bool FRivermaxManager::IsValidIP(const FString& InSourceIP) const
+	{
+		return DeviceFinder->IsValidIP(InSourceIP);
+	}
 }
 
 
diff --git a/Engine/Plugins/VirtualProduction/Rivermax/RivermaxCore/Source/RivermaxCore/Private/RivermaxManager.h b/Engine/Plugins/VirtualProduction/Rivermax/RivermaxCore/Source/RivermaxCore/Private/RivermaxManager.h
--- a/Engine/Plugins/VirtualProduction/Rivermax/RivermaxCore/Source/RivermaxCore/Private/RivermaxManager.h
+++ b/Engine/Plugins/VirtualProduction/Rivermax/RivermaxCore/Source/RivermaxCore/Private/RivermaxManager.h
@@ -4,9 +4,10 @@
 
 #include "IRivermaxManager.h"
 
-
 namespace UE::RivermaxCore::Private
 {
+	class FRivermaxDeviceFinder;
+	
 	class FRivermaxManager : public UE::RivermaxCore::IRivermaxManager
 	{
 	public:
@@ -16,14 +17,28 @@ namespace UE::RivermaxCore::Private
 	public:
 		//~ Begin IRivermaxManager interface
 		virtual bool IsInitialized() const override;
+		virtual uint64 GetTime() const override;
+		virtual ERivermaxTimeSource GetTimeSource() const override;
+		virtual TConstArrayView<FRivermaxDeviceInfo> GetDevices() const override;
+		virtual bool GetMatchingDevice(const FString& InSourceIP, FString& OutDeviceIP) const override;
+		virtual bool IsValidIP(const FString& InSourceIP) const override;
 		//~ End IRivermaxManager interface
 
 	private:
 		bool LoadRivermaxLibrary();
+		void InitializeLibrary();
+		void FindDevices();
 
 	private:
+
+		/** Whether the library has been initialized and is usable */
 		bool bIsInitialized = false;
 		void* LibraryHandle = nullptr;
+
+		/** Current time source */
+		ERivermaxTimeSource TimeSource = ERivermaxTimeSource::PTP;
+
+		TUniquePtr<FRivermaxDeviceFinder> DeviceFinder;
 	};
 }
 
diff --git a/Engine/Plugins/VirtualProduction/Rivermax/RivermaxCore/Source/RivermaxCore/Private/RivermaxSettings.cpp b/Engine/Plugins/VirtualProduction/Rivermax/RivermaxCore/Source/RivermaxCore/Private/RivermaxSettings.cpp
new file mode 100644
--- /dev/null
+++ b/Engine/Plugins/VirtualProduction/Rivermax/RivermaxCore/Source/RivermaxCore/Private/RivermaxSettings.cpp
@@ -0,0 +1,16 @@
+// Copyright Epic Games, Inc. All Rights Reserved.
+
+#include "RivermaxSettings.h"
+
+
+FName URivermaxSettings::GetCategoryName() const
+{
+	return TEXT("Plugins");
+}
+
+#if WITH_EDITOR
+FText URivermaxSettings::GetSectionText() const
+{
+	return NSLOCTEXT("RivermaxCorePlugin", "RivermaxSettingsSection", "Rivermax");
+}
+#endif //WITH_EDITOR
diff --git a/Engine/Plugins/VirtualProduction/Rivermax/RivermaxCore/Source/RivermaxCore/Public/IRivermaxManager.h b/Engine/Plugins/VirtualProduction/Rivermax/RivermaxCore/Source/RivermaxCore/Public/IRivermaxManager.h
--- a/Engine/Plugins/VirtualProduction/Rivermax/RivermaxCore/Source/RivermaxCore/Public/IRivermaxManager.h
+++ b/Engine/Plugins/VirtualProduction/Rivermax/RivermaxCore/Source/RivermaxCore/Public/IRivermaxManager.h
@@ -4,8 +4,16 @@
 
 #include "CoreMinimal.h"
 
+#include "RivermaxSettings.h"
+
 namespace UE::RivermaxCore
 {
+	struct FRivermaxDeviceInfo
+	{
+		FString Description;
+		FString InterfaceAddress;
+	};
+
 	/**
 	 * Not doing much at the moment but should be used as a central point to register every stream created
 	 * and keep track of stats. Also manages initialization of the library.
@@ -19,6 +27,25 @@ namespace UE::RivermaxCore
 		
 		/** Returns true if Rivermax has been initialized and is usable */
 		virtual bool IsInitialized() const = 0;
+
+		/** 
+		 * Returns current time in nanoseconds for Rivermax
+		 * If PTP is available, it will be used
+		 * If not, FPlatform::Seconds will be used
+		 */
+		virtual uint64 GetTime() const = 0;
+
+		/** Returns where time is coming from */
+		virtual ERivermaxTimeSource GetTimeSource() const = 0;
+
+		/** Returns list of found devices that can be used */
+		virtual TConstArrayView<FRivermaxDeviceInfo> GetDevices() const = 0;
+
+		/** Returns a matching device interface address for the provided source IP */
+		virtual bool GetMatchingDevice(const FString& InSourceIP, FString& OutDeviceIP) const = 0;
+
+		/** Returns true if provided IP string is in a valid format */
+		virtual bool IsValidIP(const FString& InSourceIP) const = 0;
 	};
 }
 
diff --git a/Engine/Plugins/VirtualProduction/Rivermax/RivermaxCore/Source/RivermaxCore/Public/RivermaxSettings.h b/Engine/Plugins/VirtualProduction/Rivermax/RivermaxCore/Source/RivermaxCore/Public/RivermaxSettings.h
new file mode 100644
--- /dev/null
+++ b/Engine/Plugins/VirtualProduction/Rivermax/RivermaxCore/Source/RivermaxCore/Public/RivermaxSettings.h
@@ -0,0 +1,42 @@
+// Copyright Epic Games, Inc. All Rights Reserved.
+
+#pragma once
+
+#include "Engine/DeveloperSettings.h"
+
+#include "RivermaxSettings.generated.h"
+
+
+UENUM()
+enum class ERivermaxTimeSource
+{
+	PTP,
+	Platform
+};
+
+/**
+ * Settings for Rivermax core plugin. 
+ */
+UCLASS(config=Engine)
+class RIVERMAXCORE_API URivermaxSettings : public UDeveloperSettings
+{
+	GENERATED_BODY()
+
+public:
+	//~ Begin UDevelopperSettings interface
+	virtual FName GetCategoryName() const;
+#if WITH_EDITOR
+	virtual FText GetSectionText() const override;
+#endif
+	//~ End UDevelopperSettings interface
+
+public:
+
+	/** Timing souce to be used by Rivermax scheduler */
+	UPROPERTY(config, EditAnywhere, Category = "Settings", meta=(ConfigRestartRequired = true))
+	ERivermaxTimeSource TimeSource = ERivermaxTimeSource::PTP;
+
+	/** Interface to use when Rivermax timing is configured to use PTP */
+	UPROPERTY(config, EditAnywhere, Category = "Settings", meta = (EditCondition = "TimeSource == ERivermaxTimeSource::PTP", ConfigRestartRequired = true))
+	FString PTPInterfaceAddress = TEXT("*.*.*.*");
+};
\ No newline at end of file
diff --git a/Engine/Plugins/VirtualProduction/Rivermax/RivermaxCore/Source/RivermaxCore/RivermaxCore.build.cs b/Engine/Plugins/VirtualProduction/Rivermax/RivermaxCore/Source/RivermaxCore/RivermaxCore.build.cs
--- a/Engine/Plugins/VirtualProduction/Rivermax/RivermaxCore/Source/RivermaxCore/RivermaxCore.build.cs
+++ b/Engine/Plugins/VirtualProduction/Rivermax/RivermaxCore/Source/RivermaxCore/RivermaxCore.build.cs
@@ -10,6 +10,7 @@ public class RivermaxCore : ModuleRules
 			new string[] {
 				"Core",
 				"CoreUObject",
+				"DeveloperSettings",
 				"MediaAssets",
 				"MediaIOCore",
 				"RivermaxLib",
@@ -20,6 +21,7 @@ public class RivermaxCore : ModuleRules
 		PrivateDependencyModuleNames.AddRange(
 			new string[] {
 				"Engine",
+				"Networking",
 				"RenderCore",
 				"RHI"
 			}
diff --git a/Engine/Plugins/VirtualProduction/Rivermax/RivermaxCore/Source/RivermaxEditor/Private/Customizations/RivermaxDeviceSelectionCustomization.cpp b/Engine/Plugins/VirtualProduction/Rivermax/RivermaxCore/Source/RivermaxEditor/Private/Customizations/RivermaxDeviceSelectionCustomization.cpp
new file mode 100644
--- /dev/null
+++ b/Engine/Plugins/VirtualProduction/Rivermax/RivermaxCore/Source/RivermaxEditor/Private/Customizations/RivermaxDeviceSelectionCustomization.cpp
@@ -0,0 +1,36 @@
+// Copyright Epic Games, Inc. All Rights Reserved.
+
+#include "Customizations/RivermaxDeviceSelectionCustomization.h"
+
+#include "DetailLayoutBuilder.h"
+#include "DetailWidgetRow.h"
+#include "Widgets/SRivermaxInterfaceComboBox.h"
+
+namespace UE::RivermaxCore::Utils
+{
+
+void SetupDeviceSelectionCustomization(int32 InObjectIndex, const FString& InitialValue, TSharedPtr<IPropertyHandle> PropertyHandle, IDetailLayoutBuilder& DetailBuilder)
+{
+	IDetailPropertyRow* Row = DetailBuilder.EditDefaultProperty(PropertyHandle);
+
+	TSharedPtr<SWidget> NameWidget;
+	TSharedPtr<SWidget> ValueWidget;
+	Row->GetDefaultWidgets(NameWidget, ValueWidget);
+	
+	Row->CustomWidget()
+		.NameContent()
+		[
+			NameWidget->AsShared()
+		]
+		.ValueContent()
+		[
+			SNew(SRivermaxInterfaceComboBox)
+			.InitialValue(InitialValue)
+			.OnIPAddressSelected_Lambda([InObjectIndex, PropertyHandle](const FString& SelectedIp)
+			{
+				PropertyHandle->SetPerObjectValue(InObjectIndex, SelectedIp);
+			})
+		];
+}
+
+}
diff --git a/Engine/Plugins/VirtualProduction/Rivermax/RivermaxCore/Source/RivermaxEditor/Private/Customizations/RivermaxSettingsDetailsCustomization.cpp b/Engine/Plugins/VirtualProduction/Rivermax/RivermaxCore/Source/RivermaxEditor/Private/Customizations/RivermaxSettingsDetailsCustomization.cpp
new file mode 100644
--- /dev/null
+++ b/Engine/Plugins/VirtualProduction/Rivermax/RivermaxCore/Source/RivermaxEditor/Private/Customizations/RivermaxSettingsDetailsCustomization.cpp
@@ -0,0 +1,29 @@
+// Copyright Epic Games, Inc. All Rights Reserved.
+
+#include "RivermaxSettingsDetailsCustomization.h"
+
+#include "Customizations/RivermaxDeviceSelectionCustomization.h"
+#include "DetailLayoutBuilder.h"
+#include "DetailWidgetRow.h"
+#include "RivermaxSettings.h"
+
+TSharedRef<IDetailCustomization> FRivermaxSettingsDetailsCustomization::MakeInstance()
+{
+	return MakeShared<FRivermaxSettingsDetailsCustomization>();
+}
+
+void FRivermaxSettingsDetailsCustomization::CustomizeDetails(IDetailLayoutBuilder& DetailBuilder)
+{
+	TArray<TWeakObjectPtr<UObject>> CustomizedObjects;
+	DetailBuilder.GetObjectsBeingCustomized(CustomizedObjects);
+
+	for (int32 ObjectIndex = 0; ObjectIndex < CustomizedObjects.Num(); ++ObjectIndex)
+	{
+		// For now, take care of both Rivermax source and output to customize their interface field
+		const TWeakObjectPtr<UObject> Obj = CustomizedObjects[ObjectIndex];
+		if (URivermaxSettings* RivermaxSettings = Cast<URivermaxSettings>(Obj.Get()))
+		{
+			UE::RivermaxCore::Utils::SetupDeviceSelectionCustomization(ObjectIndex, RivermaxSettings->PTPInterfaceAddress, DetailBuilder.GetProperty(GET_MEMBER_NAME_CHECKED(URivermaxSettings, PTPInterfaceAddress)), DetailBuilder);
+		}
+	}
+}
diff --git a/Engine/Plugins/VirtualProduction/Rivermax/RivermaxCore/Source/RivermaxEditor/Private/Customizations/RivermaxSettingsDetailsCustomization.h b/Engine/Plugins/VirtualProduction/Rivermax/RivermaxCore/Source/RivermaxEditor/Private/Customizations/RivermaxSettingsDetailsCustomization.h
new file mode 100644
--- /dev/null
+++ b/Engine/Plugins/VirtualProduction/Rivermax/RivermaxCore/Source/RivermaxEditor/Private/Customizations/RivermaxSettingsDetailsCustomization.h
@@ -0,0 +1,19 @@
+// Copyright Epic Games, Inc. All Rights Reserved.
+
+#pragma once
+
+#include "IDetailCustomization.h"
+
+
+class IPropertyHandle;
+
+class FRivermaxSettingsDetailsCustomization : public IDetailCustomization
+{
+public:
+	
+	static TSharedRef<IDetailCustomization> MakeInstance();
+	
+	//~ Begin IDetailCustomization Interface
+	virtual void CustomizeDetails(IDetailLayoutBuilder& DetailBuilder) override;
+	//~ End IDetailCustomization Interface
+};
diff --git a/Engine/Plugins/VirtualProduction/Rivermax/RivermaxCore/Source/RivermaxEditor/Private/RivermaxEditorModule.cpp b/Engine/Plugins/VirtualProduction/Rivermax/RivermaxCore/Source/RivermaxEditor/Private/RivermaxEditorModule.cpp
new file mode 100644
--- /dev/null
+++ b/Engine/Plugins/VirtualProduction/Rivermax/RivermaxCore/Source/RivermaxEditor/Private/RivermaxEditorModule.cpp
@@ -0,0 +1,35 @@
+// Copyright Epic Games, Inc. All Rights Reserved.
+
+#include "RivermaxEditorModule.h"
+
+#include "Customizations/RivermaxSettingsDetailsCustomization.h"
+#include "PropertyEditorModule.h"
+#include "RivermaxSettings.h"
+
+
+void FRivermaxEditorModule::StartupModule()
+{
+	FPropertyEditorModule& PropertyModule = FModuleManager::GetModuleChecked<FPropertyEditorModule>("PropertyEditor");
+
+	const auto AddCustomization = [&PropertyModule, this](const FName ClassName, FOnGetDetailCustomizationInstance Customizor)
+	{
+		PropertyModule.RegisterCustomClassLayout(ClassName, Customizor);
+		CustomizedClasses.Add(ClassName);
+	};
+
+	AddCustomization(URivermaxSettings::StaticClass()->GetFName(), FOnGetDetailCustomizationInstance::CreateStatic(&FRivermaxSettingsDetailsCustomization::MakeInstance));
+}
+
+void FRivermaxEditorModule::ShutdownModule()
+{
+	FPropertyEditorModule* PropertyModule = FModuleManager::GetModulePtr<FPropertyEditorModule>("PropertyEditor");
+	if (PropertyModule)
+	{
+		for (FName ClassName : CustomizedClasses)
+		{
+			PropertyModule->UnregisterCustomClassLayout(ClassName);
+		}
+	}
+}
+
+IMPLEMENT_MODULE(FRivermaxEditorModule, RivermaxEditor);
diff --git a/Engine/Plugins/VirtualProduction/Rivermax/RivermaxCore/Source/RivermaxEditor/Private/RivermaxEditorModule.h b/Engine/Plugins/VirtualProduction/Rivermax/RivermaxCore/Source/RivermaxEditor/Private/RivermaxEditorModule.h
new file mode 100644
--- /dev/null
+++ b/Engine/Plugins/VirtualProduction/Rivermax/RivermaxCore/Source/RivermaxEditor/Private/RivermaxEditorModule.h
@@ -0,0 +1,21 @@
+// Copyright Epic Games, Inc. All Rights Reserved.
+
+#pragma once
+
+#include "Modules/ModuleManager.h"
+
+
+/**
+ */
+class FRivermaxEditorModule : public IModuleInterface
+{
+public:
+	//~ Begin IModuleInterface
+	virtual void StartupModule() override;
+	virtual void ShutdownModule() override;
+	//~ End IModuleInterface
+
+private:
+	TArray<FName> CustomizedClasses;
+};
+
diff --git a/Engine/Plugins/VirtualProduction/Rivermax/RivermaxCore/Source/RivermaxEditor/Private/Widgets/SRivermaxInterfaceComboBox.cpp b/Engine/Plugins/VirtualProduction/Rivermax/RivermaxCore/Source/RivermaxEditor/Private/Widgets/SRivermaxInterfaceComboBox.cpp
new file mode 100644
--- /dev/null
+++ b/Engine/Plugins/VirtualProduction/Rivermax/RivermaxCore/Source/RivermaxEditor/Private/Widgets/SRivermaxInterfaceComboBox.cpp
@@ -0,0 +1,125 @@
+// Copyright Epic Games, Inc. All Rights Reserved.
+
+#include "Widgets/SRivermaxInterfaceComboBox.h"
+
+#include "EditorStyleSet.h"
+#include "IRivermaxCoreModule.h"
+#include "IRivermaxManager.h"
+#include "Styling/AppStyle.h"
+#include "Widgets/Images/SImage.h"
+#include "Widgets/Input/SEditableComboBox.h"
+#include "Widgets/Input/SEditableTextBox.h"
+
+
+#define LOCTEXT_NAMESPACE "RivermaxComboBox"
+
+
+namespace UE::RivermaxEditor::Private
+{
+	static TArray<TSharedPtr<FString>> GetRivermaxNetworkInterfaceCardIPs()
+	{
+		TArray<TSharedPtr<FString>> RivermaxDeviceIPs;
+
+		IRivermaxCoreModule& RivermaxModule = FModuleManager::LoadModuleChecked<IRivermaxCoreModule>(TEXT("RivermaxCore"));
+		for(const UE::RivermaxCore::FRivermaxDeviceInfo& Device : RivermaxModule.GetRivermaxManager()->GetDevices())
+		{
+			RivermaxDeviceIPs.Add(MakeShared<FString>(Device.InterfaceAddress));
+		}
+		
+		return RivermaxDeviceIPs;
+	}
+}
+
+void SRivermaxInterfaceComboBox::Construct(const FArguments& InArgs)
+{
+	OnIPAddressSelectedDelegate = InArgs._OnIPAddressSelected;
+	RivermaxAdapterAddressSource = UE::RivermaxEditor::Private::GetRivermaxNetworkInterfaceCardIPs();
+	ChildSlot
+	[
+		SAssignNew(HorizontalBoxContainer, SHorizontalBox)
+
+		+ SHorizontalBox::Slot()
+		.AutoWidth()
+		.VAlign(VAlign_Center)
+		[
+			SAssignNew(RivermaxAdapterAddressComboBox, SComboBox<TSharedPtr<FString>>)
+			.OptionsSource(&RivermaxAdapterAddressSource)
+			.OnGenerateWidget(this, &SRivermaxInterfaceComboBox::GenerateRivermaxAdaptersComboBoxEntry)
+			.OnSelectionChanged(this, &SRivermaxInterfaceComboBox::OnIpAddressSelected)
+			.Content()
+			[
+				SAssignNew(IPAddressEditableTextBlock, SEditableText)
+				.Font(FAppStyle::GetFontStyle(TEXT("PropertyWindow.NormalFont")))
+				.Text(FText::FromString(InArgs._InitialValue))
+				.OnTextCommitted(this, &SRivermaxInterfaceComboBox::OnIPAddressTextCommmited)
+			]
+		]
+
+		+ SHorizontalBox::Slot()
+		.FillWidth(1.0f)
+		.HAlign(HAlign_Left)
+		.VAlign(VAlign_Center)
+		.Padding(4.0f, 0.0f, 0.0f, 0.0f)
+		[
+			SAssignNew(IPAddressWarningImage, SImage)
+		]
+	];
+	
+	UpdateIPWarning();
+}
+
+void SRivermaxInterfaceComboBox::OnIpAddressSelected(TSharedPtr<FString> InAddress, ESelectInfo::Type InType)
+{
+	TSharedPtr<FString> SelectedIPAddress = RivermaxAdapterAddressComboBox->GetSelectedItem();
+	if (SelectedIPAddress.IsValid())
+	{
+		IPAddressEditableTextBlock->SetText(FText::FromString(*SelectedIPAddress));
+		OnIPAddressSelectedDelegate.ExecuteIfBound(*InAddress.Get());
+		UpdateIPWarning();
+	}
+}
+
+void SRivermaxInterfaceComboBox::OnIPAddressTextCommmited(const FText& Text, ETextCommit::Type)
+{
+	OnIPAddressSelectedDelegate.ExecuteIfBound(Text.ToString());
+	UpdateIPWarning();
+}
+
+TSharedRef<SWidget> SRivermaxInterfaceComboBox::GenerateRivermaxAdaptersComboBoxEntry(TSharedPtr<FString> InAddress)
+{
+	return
+		SNew(STextBlock)
+		.Text(FText::FromString(*InAddress))
+		.Font(FAppStyle::GetFontStyle(TEXT("PropertyWindow.NormalFont")));
+}
+
+void SRivermaxInterfaceComboBox::UpdateIPWarning()
+{
+	const FString CurrentIP = IPAddressEditableTextBlock->GetText().ToString();
+	IRivermaxCoreModule& Module = FModuleManager::GetModuleChecked<IRivermaxCoreModule>("RivermaxCore");
+	if (Module.GetRivermaxManager()->IsValidIP(CurrentIP))
+	{
+		FString MatchingDevice;
+		if (Module.GetRivermaxManager()->GetMatchingDevice(CurrentIP, MatchingDevice))
+		{
+			HorizontalBoxContainer->SetToolTipText(FText::Format(LOCTEXT("MatchingDeviceFound", "Provided IP resolved to Rivermax device {0}"), FText::FromString(MatchingDevice)));
+			IPAddressWarningImage->SetVisibility(EVisibility::Collapsed);
+
+		}
+		else
+		{
+			IPAddressWarningImage->SetImage(FAppStyle::Get().GetBrush("Icons.WarningWithColor"));
+			HorizontalBoxContainer->SetToolTipText(LOCTEXT("MatchingDeviceNotFound", "Could not find Rivermax device matching provided IP"));
+			IPAddressWarningImage->SetVisibility(EVisibility::Visible);
+		}
+	}
+	else
+	{
+		IPAddressWarningImage->SetImage(FAppStyle::Get().GetBrush("Icons.ErrorWithColor"));
+		HorizontalBoxContainer->SetToolTipText(LOCTEXT("InvalidIPProvided", "Invalid IP format provided. Must be of the form x.x.x.x with * accepted as wildcard."));
+		IPAddressWarningImage->SetVisibility(EVisibility::Visible);
+	}
+}
+
+#undef LOCTEXT_NAMESPACE
+
diff --git a/Engine/Plugins/VirtualProduction/Rivermax/RivermaxCore/Source/RivermaxEditor/Private/Widgets/SRivermaxInterfaceComboBox.h b/Engine/Plugins/VirtualProduction/Rivermax/RivermaxCore/Source/RivermaxEditor/Private/Widgets/SRivermaxInterfaceComboBox.h
new file mode 100644
--- /dev/null
+++ b/Engine/Plugins/VirtualProduction/Rivermax/RivermaxCore/Source/RivermaxEditor/Private/Widgets/SRivermaxInterfaceComboBox.h
@@ -0,0 +1,63 @@
+// Copyright Epic Games, Inc. All Rights Reserved.
+
+#pragma once
+
+#include "Widgets/SCompoundWidget.h"
+
+#include "Widgets/Input/SComboBox.h"
+
+class SEditableText;
+class SEditableTextBox;
+class SImage;
+class SHorizontalBox;
+
+/** 
+ * Widget used to pick IP address of Rivermax / Mellanox device
+ */
+class RIVERMAXEDITOR_API SRivermaxInterfaceComboBox : public SCompoundWidget
+{
+public:
+
+	DECLARE_DELEGATE_OneParam(FOnChooseIP, FString)
+	
+	SLATE_BEGIN_ARGS(SRivermaxInterfaceComboBox)
+	{}
+	SLATE_ARGUMENT(FString, InitialValue)
+		SLATE_EVENT(FOnChooseIP, OnIPAddressSelected)
+	SLATE_END_ARGS()
+
+	void Construct(const FArguments& InArgs);
+	
+private:
+	
+	/** Handles changes in the local adapter address combo box */
+	void OnIpAddressSelected(TSharedPtr<FString> InAddress, ESelectInfo::Type InType);
+
+	/** Called when the IP Address was commited in the editable text block */
+	void OnIPAddressTextCommmited(const FText&, ETextCommit::Type);
+
+	/** Generates an entry in the local adapter address combo box */
+	TSharedRef<SWidget> GenerateRivermaxAdaptersComboBoxEntry(TSharedPtr<FString> InAddress);
+
+	/** Updates image showing warning if IP is invalid */
+	void UpdateIPWarning();
+
+private:
+
+	TArray<TSharedPtr<FString>> RivermaxAdapterAddressSource;
+
+	TSharedPtr<SComboBox<TSharedPtr<FString>>> RivermaxAdapterAddressComboBox;
+
+	/** Text box shown on top of the local adapter address address combo box */
+	TSharedPtr<SEditableText> IPAddressEditableTextBlock;
+
+	/** Warning images displayed when IP is not valid or no device matches resolved IP */
+	TSharedPtr<SImage> IPAddressWarningImage;
+	
+	/** Horizontal box containing combo button, text and image */
+	TSharedPtr<SHorizontalBox> HorizontalBoxContainer;
+
+	/** Delegate executed when a local IP address was selected */
+	FOnChooseIP OnIPAddressSelectedDelegate;
+};
+
diff --git a/Engine/Plugins/VirtualProduction/Rivermax/RivermaxCore/Source/RivermaxEditor/Public/Customizations/RivermaxDeviceSelectionCustomization.h b/Engine/Plugins/VirtualProduction/Rivermax/RivermaxCore/Source/RivermaxEditor/Public/Customizations/RivermaxDeviceSelectionCustomization.h
new file mode 100644
--- /dev/null
+++ b/Engine/Plugins/VirtualProduction/Rivermax/RivermaxCore/Source/RivermaxEditor/Public/Customizations/RivermaxDeviceSelectionCustomization.h
@@ -0,0 +1,13 @@
+// Copyright Epic Games, Inc. All Rights Reserved.
+
+#pragma once
+
+#include "CoreMinimal.h"
+#include "IDetailCustomization.h"
+
+class IPropertyHandle;
+
+namespace UE::RivermaxCore::Utils
+{
+	RIVERMAXEDITOR_API void SetupDeviceSelectionCustomization(int32 InObjectIndex, const FString& InitialValue, TSharedPtr<IPropertyHandle> PropertyHandle, IDetailLayoutBuilder& DetailBuilder);
+}
diff --git a/Engine/Plugins/VirtualProduction/Rivermax/RivermaxCore/Source/RivermaxEditor/RivermaxEditor.build.cs b/Engine/Plugins/VirtualProduction/Rivermax/RivermaxCore/Source/RivermaxEditor/RivermaxEditor.build.cs
new file mode 100644
--- /dev/null
+++ b/Engine/Plugins/VirtualProduction/Rivermax/RivermaxCore/Source/RivermaxEditor/RivermaxEditor.build.cs
@@ -0,0 +1,40 @@
+// Copyright Epic Games, Inc. All Rights Reserved.
+
+using UnrealBuildTool;
+
+public class RivermaxEditor : ModuleRules
+{
+	public RivermaxEditor(ReadOnlyTargetRules Target) : base(Target)
+	{
+		PublicDependencyModuleNames.AddRange(
+			new string[] {
+				"Core",
+				"PropertyEditor",
+				"RivermaxCore",
+			}
+		);
+
+		PrivateDependencyModuleNames.AddRange(
+			new string[] {
+				"CoreUObject",
+				"EditorStyle",
+				"InputCore",
+				"Slate",
+				"SlateCore",
+			}
+		);
+
+		if (Target.bBuildEditor)
+		{
+			PublicDependencyModuleNames.AddRange(
+				new string[] {
+				}
+			);
+
+			PrivateDependencyModuleNames.AddRange(
+				new string[] {
+				}
+			);
+		}
+	}
+}
diff --git a/Engine/Plugins/VirtualProduction/Rivermax/RivermaxMedia/Source/RivermaxMedia/Private/RivermaxMediaCapture.cpp b/Engine/Plugins/VirtualProduction/Rivermax/RivermaxMedia/Source/RivermaxMedia/Private/RivermaxMediaCapture.cpp
--- a/Engine/Plugins/VirtualProduction/Rivermax/RivermaxMedia/Source/RivermaxMedia/Private/RivermaxMediaCapture.cpp
+++ b/Engine/Plugins/VirtualProduction/Rivermax/RivermaxMedia/Source/RivermaxMedia/Private/RivermaxMediaCapture.cpp
@@ -3,6 +3,7 @@
 #include "RivermaxMediaCapture.h"
 
 #include "IRivermaxCoreModule.h"
+#include "IRivermaxManager.h"
 #include "IRivermaxOutputStream.h"
 #include "RenderGraphUtils.h"
 #include "RivermaxMediaLog.h"
@@ -189,8 +190,21 @@ bool URivermaxMediaCapture::ConfigureStream(URivermaxMediaOutput* InMediaOutput,
 {
 	using namespace UE::RivermaxCore;
 
+	// Resolve interface address
+	IRivermaxCoreModule* Module = FModuleManager::GetModulePtr<IRivermaxCoreModule>("RivermaxCore");
+	if (Module == nullptr)
+	{
+		return false;
+	}
+
+	const bool bFoundDevice = Module->GetRivermaxManager()->GetMatchingDevice(InMediaOutput->InterfaceAddress, OutOptions.InterfaceAddress);
+	if (bFoundDevice == false)
+	{
+		UE_LOG(LogRivermaxMedia, Error, TEXT("Could not find a matching interface for IP '%s'"), *InMediaOutput->InterfaceAddress);
+		return false;
+	}
+
 	OutOptions.StreamAddress = InMediaOutput->StreamAddress;
-	OutOptions.InterfaceAddress = InMediaOutput->InterfaceAddress;
 	OutOptions.Port = InMediaOutput->Port;
 	OutOptions.Resolution = InMediaOutput->Resolution;
 	OutOptions.FrameRate = InMediaOutput->FrameRate;
diff --git a/Engine/Plugins/VirtualProduction/Rivermax/RivermaxMedia/Source/RivermaxMedia/Private/RivermaxMediaPlayer.cpp b/Engine/Plugins/VirtualProduction/Rivermax/RivermaxMedia/Source/RivermaxMedia/Private/RivermaxMediaPlayer.cpp
--- a/Engine/Plugins/VirtualProduction/Rivermax/RivermaxMedia/Source/RivermaxMedia/Private/RivermaxMediaPlayer.cpp
+++ b/Engine/Plugins/VirtualProduction/Rivermax/RivermaxMedia/Source/RivermaxMedia/Private/RivermaxMediaPlayer.cpp
@@ -4,6 +4,7 @@
 
 #include "IMediaEventSink.h"
 #include "IRivermaxCoreModule.h"
+#include "IRivermaxManager.h"
 #include "MediaIOCoreEncodeTime.h"
 #include "MediaIOCoreFileWriter.h"
 #include "MediaIOCoreSamples.h"
@@ -277,7 +278,21 @@ namespace UE::RivermaxMedia
 
 	bool FRivermaxMediaPlayer::ConfigureStream(const IMediaOptions* Options)
 	{
-		StreamOptions.InterfaceAddress = Options->GetMediaOption(RivermaxMediaOption::InterfaceAddress, FString());
+		// Resolve interface address
+		IRivermaxCoreModule* Module = FModuleManager::GetModulePtr<IRivermaxCoreModule>("RivermaxCore");
+		if (Module == nullptr)
+		{
+			return false;
+		}
+
+		const FString DesiredInterface = Options->GetMediaOption(RivermaxMediaOption::InterfaceAddress, FString());
+		const bool bFoundDevice = Module->GetRivermaxManager()->GetMatchingDevice(DesiredInterface, StreamOptions.InterfaceAddress);
+		if (bFoundDevice == false)
+		{
+			UE_LOG(LogRivermaxMedia, Error, TEXT("Could not find a matching interface for IP '%s'"), *DesiredInterface);
+			return false;
+		}
+
 		StreamOptions.StreamAddress = Options->GetMediaOption(RivermaxMediaOption::StreamAddress, FString());
 		StreamOptions.Port = Options->GetMediaOption(RivermaxMediaOption::Port, (int64)0);
 		StreamOptions.Resolution = VideoTrackFormat.Dim;
diff --git a/Engine/Plugins/VirtualProduction/Rivermax/RivermaxMedia/Source/RivermaxMediaEditor/Private/Customization/RivermaxMediaDetailsCustomization.cpp b/Engine/Plugins/VirtualProduction/Rivermax/RivermaxMedia/Source/RivermaxMediaEditor/Private/Customization/RivermaxMediaDetailsCustomization.cpp
new file mode 100644
--- /dev/null
+++ b/Engine/Plugins/VirtualProduction/Rivermax/RivermaxMedia/Source/RivermaxMediaEditor/Private/Customization/RivermaxMediaDetailsCustomization.cpp
@@ -0,0 +1,34 @@
+// Copyright Epic Games, Inc. All Rights Reserved.
+
+#include "RivermaxMediaDetailsCustomization.h"
+
+#include "Customizations/RivermaxDeviceSelectionCustomization.h"
+#include "DetailLayoutBuilder.h"
+#include "DetailWidgetRow.h"
+#include "RivermaxMediaOutput.h"
+#include "RivermaxMediaSource.h"
+
+TSharedRef<IDetailCustomization> FRivermaxMediaDetailsCustomization::MakeInstance()
+{
+	return MakeShared<FRivermaxMediaDetailsCustomization>();
+}
+
+void FRivermaxMediaDetailsCustomization::CustomizeDetails(IDetailLayoutBuilder& DetailBuilder)
+{
+	TArray<TWeakObjectPtr<UObject>> CustomizedObjects;
+	DetailBuilder.GetObjectsBeingCustomized(CustomizedObjects);
+
+	for (int32 ObjectIndex = 0; ObjectIndex < CustomizedObjects.Num(); ++ObjectIndex)
+	{
+		// For now, take care of both Rivermax source and output to customize their interface field
+		const TWeakObjectPtr<UObject> Obj = CustomizedObjects[ObjectIndex];
+		if (URivermaxMediaOutput* Output = Cast<URivermaxMediaOutput>(Obj.Get()))
+		{
+			UE::RivermaxCore::Utils::SetupDeviceSelectionCustomization(ObjectIndex, Output->InterfaceAddress, DetailBuilder.GetProperty(GET_MEMBER_NAME_CHECKED(URivermaxMediaOutput, InterfaceAddress)), DetailBuilder);
+		}
+		else if (URivermaxMediaSource* Source = Cast<URivermaxMediaSource>(Obj.Get()))
+		{
+			UE::RivermaxCore::Utils::SetupDeviceSelectionCustomization(ObjectIndex, Source->InterfaceAddress, DetailBuilder.GetProperty(GET_MEMBER_NAME_CHECKED(URivermaxMediaSource, InterfaceAddress)), DetailBuilder);
+		}
+	}
+}
diff --git a/Engine/Plugins/VirtualProduction/Rivermax/RivermaxMedia/Source/RivermaxMediaEditor/Private/Customization/RivermaxMediaDetailsCustomization.h b/Engine/Plugins/VirtualProduction/Rivermax/RivermaxMedia/Source/RivermaxMediaEditor/Private/Customization/RivermaxMediaDetailsCustomization.h
new file mode 100644
--- /dev/null
+++ b/Engine/Plugins/VirtualProduction/Rivermax/RivermaxMedia/Source/RivermaxMediaEditor/Private/Customization/RivermaxMediaDetailsCustomization.h
@@ -0,0 +1,19 @@
+// Copyright Epic Games, Inc. All Rights Reserved.
+
+#pragma once
+
+#include "IDetailCustomization.h"
+
+
+class IPropertyHandle;
+
+class FRivermaxMediaDetailsCustomization : public IDetailCustomization
+{
+public:
+	
+	static TSharedRef<IDetailCustomization> MakeInstance();
+	
+	//~ Begin IDetailCustomization Interface
+	virtual void CustomizeDetails(IDetailLayoutBuilder& DetailBuilder) override;
+	//~ End IDetailCustomization Interface
+};
diff --git a/Engine/Plugins/VirtualProduction/Rivermax/RivermaxMedia/Source/RivermaxMediaEditor/Private/RivermaxMediaEditorModule.cpp b/Engine/Plugins/VirtualProduction/Rivermax/RivermaxMedia/Source/RivermaxMediaEditor/Private/RivermaxMediaEditorModule.cpp
--- a/Engine/Plugins/VirtualProduction/Rivermax/RivermaxMedia/Source/RivermaxMediaEditor/Private/RivermaxMediaEditorModule.cpp
+++ b/Engine/Plugins/VirtualProduction/Rivermax/RivermaxMedia/Source/RivermaxMediaEditor/Private/RivermaxMediaEditorModule.cpp
@@ -2,4 +2,59 @@
 
 #include "Modules/ModuleManager.h"
 
-IMPLEMENT_MODULE(FDefaultModuleImpl, RivermaxMediaEditor);
+#include "PropertyEditorModule.h"
+#include "RivermaxMediaOutput.h"
+#include "RivermaxMediaSource.h"
+#include "Customization/RivermaxMediaDetailsCustomization.h"
+
+class FRivermaxEditorModule : public IModuleInterface
+{
+public:
+
+	//~ Begin IModuleInterface
+	virtual void StartupModule() override
+	{
+		RegisterCustomizations();
+	}
+	
+	virtual void ShutdownModule() override
+	{
+		UnregisterCustomizations();
+	}
+	//~End IModuleInterface
+
+private:
+	/** Registers detail customization for Rivermax media types */
+	void RegisterCustomizations()
+	{
+		FPropertyEditorModule& PropertyModule = FModuleManager::GetModuleChecked<FPropertyEditorModule>("PropertyEditor");
+
+		const auto AddCustomization = [&PropertyModule, this](const FName ClassName, FOnGetDetailCustomizationInstance Customizor)
+		{
+			PropertyModule.RegisterCustomClassLayout(ClassName, Customizor);
+			CustomizedClasses.Add(ClassName);
+		};
+
+		AddCustomization(URivermaxMediaSource::StaticClass()->GetFName(), FOnGetDetailCustomizationInstance::CreateStatic(&FRivermaxMediaDetailsCustomization::MakeInstance));
+		AddCustomization(URivermaxMediaOutput::StaticClass()->GetFName(), FOnGetDetailCustomizationInstance::CreateStatic(&FRivermaxMediaDetailsCustomization::MakeInstance));
+	}
+
+	/** Unregisters registered customizations */
+	void UnregisterCustomizations()
+	{
+		FPropertyEditorModule* PropertyModule = FModuleManager::GetModulePtr<FPropertyEditorModule>("PropertyEditor");
+		if (PropertyModule)
+		{
+			for (FName ClassName : CustomizedClasses)
+			{
+				PropertyModule->UnregisterCustomClassLayout(ClassName);
+			}
+		}
+	}
+
+private:
+	TArray<FName> CustomizedClasses;
+};
+
+
+IMPLEMENT_MODULE(FRivermaxEditorModule, RivermaxMediaEditor);
diff --git a/Engine/Plugins/VirtualProduction/Rivermax/RivermaxMedia/Source/RivermaxMediaEditor/RivermaxMediaEditor.build.cs b/Engine/Plugins/VirtualProduction/Rivermax/RivermaxMedia/Source/RivermaxMediaEditor/RivermaxMediaEditor.build.cs
--- a/Engine/Plugins/VirtualProduction/Rivermax/RivermaxMedia/Source/RivermaxMediaEditor/RivermaxMediaEditor.build.cs
+++ b/Engine/Plugins/VirtualProduction/Rivermax/RivermaxMedia/Source/RivermaxMediaEditor/RivermaxMediaEditor.build.cs
@@ -1,5 +1,6 @@
 // Copyright Epic Games, Inc. All Rights Reserved.
 
+using System.IO;
 using UnrealBuildTool;
 
 public class RivermaxMediaEditor : ModuleRules
@@ -22,6 +23,8 @@ public class RivermaxMediaEditor : ModuleRules
 				"MediaAssets",
 				"MediaIOCore",
 				"MediaIOEditor",
+				"RivermaxCore",
+				"RivermaxEditor",
 				"RivermaxMedia",
 				"Slate",
 				"SlateCore",
