- Adding PTP alignment utils
- Adding missing PTP scheduling support and associated genlock
- Fixing 2110 output hitching with engine by supporting cotinuous mode

diff --git a/Engine/Plugins/VirtualProduction/Rivermax/RivermaxCore/Source/RivermaxCore/Private/RivermaxOutputFrame.cpp b/Engine/Plugins/VirtualProduction/Rivermax/RivermaxCore/Source/RivermaxCore/Private/RivermaxOutputFrame.cpp
--- a/Engine/Plugins/VirtualProduction/Rivermax/RivermaxCore/Source/RivermaxCore/Private/RivermaxOutputFrame.cpp
+++ b/Engine/Plugins/VirtualProduction/Rivermax/RivermaxCore/Source/RivermaxCore/Private/RivermaxOutputFrame.cpp
@@ -26,11 +26,16 @@ namespace UE::RivermaxCore::Private
 
 	void FRivermaxOutputFrame::Clear()
 	{
-		FrameIdentifier = InvalidIdentifier;
 		PacketCounter = 0;
 		LineNumber = 0;
 		SRDOffset = 0;
-		TimestampTicks = 0.0;
+	}
+
+	void FRivermaxOutputFrame::Reset()
+	{
+		FrameIdentifier = InvalidIdentifier;
 		bIsVideoBufferReady = false;
+		Clear();
 	}
+
 }
diff --git a/Engine/Plugins/VirtualProduction/Rivermax/RivermaxCore/Source/RivermaxCore/Private/RivermaxOutputFrame.h b/Engine/Plugins/VirtualProduction/Rivermax/RivermaxCore/Source/RivermaxCore/Private/RivermaxOutputFrame.h
--- a/Engine/Plugins/VirtualProduction/Rivermax/RivermaxCore/Source/RivermaxCore/Private/RivermaxOutputFrame.h
+++ b/Engine/Plugins/VirtualProduction/Rivermax/RivermaxCore/Source/RivermaxCore/Private/RivermaxOutputFrame.h
@@ -22,9 +22,12 @@ namespace UE::RivermaxCore::Private
 		/** Whether this frame is ready to be sent */
 		bool IsReadyToBeSent() const;
 
-		/** Reset internal counters of this frame */
+		/** Reset internal counters to make it resendable */
 		void Clear();
 
+		/** Reset identifier and internal counters to make it available */
+		void Reset();
+
 		/** Index of the frame used as unique identifier */
 		uint32 FrameIndex = 0;
 
diff --git a/Engine/Plugins/VirtualProduction/Rivermax/RivermaxCore/Source/RivermaxCore/Private/RivermaxOutputStream.cpp b/Engine/Plugins/VirtualProduction/Rivermax/RivermaxCore/Source/RivermaxCore/Private/RivermaxOutputStream.cpp
--- a/Engine/Plugins/VirtualProduction/Rivermax/RivermaxCore/Source/RivermaxCore/Private/RivermaxOutputStream.cpp
+++ b/Engine/Plugins/VirtualProduction/Rivermax/RivermaxCore/Source/RivermaxCore/Private/RivermaxOutputStream.cpp
@@ -7,81 +7,18 @@
 #include "IRivermaxManager.h"
 #include "Misc/ByteSwap.h"
 #include "RivermaxLog.h"
+#include "RivermaxPTPUtils.h"
 #include "RivermaxTypes.h"
 #include "RivermaxUtils.h"
 
 
 namespace UE::RivermaxCore::Private
 {
-	void CalculateStreamTime(const FRivermaxOutputStreamMemory& InStreamMemory, const FRivermaxOutputStreamData& InStreamData, const FRivermaxStreamOptions& InMediaOptions, const double InSampleRate, double& InOutTimeNs, double& OutTimestampTick)
-	{
-		using namespace UE::RivermaxCore::Private::Utils;
+	static TAutoConsoleVariable<int32> CVarRivermaxWakeupOffset(
+		TEXT("Rivermax.AlignmentOffset"), 0,
+		TEXT("Offset from alignment point to wake up in nanoseconds"),
+		ECVF_Default);
 
-		double FrameIntervalNs = InStreamData.FrameFieldTimeIntervalNs;
-		const bool bIsProgressive = true;//todo MediaConfiguration.IsProgressive() 
-		const ERivermaxStreamType StreamType = ERivermaxStreamType::VIDEO_2110_20_STREAM; //todo
-		if (bIsProgressive == false)
-		{
-			FrameIntervalNs *= 2;
-		}
-		const uint64 FrameAtTime = (uint64)(InOutTimeNs / FrameIntervalNs + 1);
-		double FirstPacketStartTimeNs = FrameAtTime * FrameIntervalNs; //next alignment point calculation
-
-		uint32 PacketsInFrameField = InStreamMemory.PacketsInFrameField;
-		if (bIsProgressive == false)
-		{
-			PacketsInFrameField *= 2;
-		}
-		
-		// See https://ieeexplore.ieee.org/stamp/stamp.jsp?tp=&arnumber=8165971 for reference
-		// To be revisited for non standard resolution/frame rate
-		if (StreamType == ERivermaxStreamType::VIDEO_2110_20_STREAM)
-		{
-			double RActive;
-			double TRODefaultMultiplier;
-			if (bIsProgressive)
-			{
-				RActive = (1080.0 / 1125.0);
-				if (InMediaOptions.Resolution.Y >= FullHDHeight)
-				{
-				 	// As defined by SMPTE 2110-21 6.3.2
-					TRODefaultMultiplier = (43.0 / 1125.0);
-				}
-				else 
-				{
-					TRODefaultMultiplier = (28.0 / 750.0);
-				}
-			}
-			else
-			{
-				if (InMediaOptions.Resolution.Y >= FullHDHeight)
-				{
-					// As defined by SMPTE 2110-21 6.3.3
-					RActive = (1080.0 / 1125.0);
-					TRODefaultMultiplier = (22.0 / 1125.0);
-				}
-				else if (InMediaOptions.Resolution.Y >= 576)
-				{
-					RActive = (576.0 / 625.0);
-					TRODefaultMultiplier = (26.0 / 625.0);
-				}
-				else 
-				{
-					RActive = (487.0 / 525.0);
-					TRODefaultMultiplier = (20.0 / 525.0);
-				}
-			}
-
-			const double TRSNano = (FrameIntervalNs * RActive) / PacketsInFrameField;
-			const double TROffset = (TRODefaultMultiplier * FrameIntervalNs) - (VideoTROModification * TRSNano);
-			FirstPacketStartTimeNs += TROffset;
-		}
-
-		constexpr double OneSecondNs = 1 * 1E9;
-		OutTimestampTick = FirstPacketStartTimeNs * (InSampleRate / OneSecondNs);
-		InOutTimeNs = FirstPacketStartTimeNs;
-	}
-	
 	uint32 FindPayloadSize(const FRivermaxStreamOptions& InOptions, uint32 InBytesPerLine)
 	{
 		using namespace UE::RivermaxCore::Private::Utils;
@@ -149,6 +86,7 @@ namespace UE::RivermaxCore::Private
 
 		Options = InOptions;
 		Listener = &InListener;
+		AlignmentPointSchedulingLeadTimeNanosec = CVarRivermaxWakeupOffset.GetValueOnGameThread();
 
 		Async(EAsyncExecution::TaskGraph, [this]()
 		{
@@ -190,9 +128,7 @@ namespace UE::RivermaxCore::Private
 					// Todo add event manager to avoid spinning
 
 					StreamData.FrameFieldTimeIntervalNs = 1E9 / Options.FrameRate.AsDecimal();
-					StreamData.SendTimeNs = (FPlatformTime::Seconds() + 1) * 1E9; //Current time + 1 sec todo why?
-					CalculateStreamTime(StreamMemory, StreamData, Options, MediaClockSampleRate, StreamData.SendTimeNs, StreamData.InitialTimestampTick);
-					StreamData.StartSendTimeNs = StreamData.SendTimeNs;
+					InitializeStreamTimingSettings();
 
 					bIsActive = true;
 					RivermaxThread.Reset(FRunnableThread::Create(this, TEXT("Rmax OutputStream Thread"), 128 * 1024, TPri_AboveNormal, FPlatformAffinity::GetPoolThreadMask()));
@@ -266,35 +202,16 @@ namespace UE::RivermaxCore::Private
 		{
 			{
 				TRACE_CPUPROFILER_EVENT_SCOPE(FRivermaxOutputStream::Wait);
-				StreamData.SendTimeNs = StreamData.StartSendTimeNs + StreamData.FrameFieldTimeIntervalNs * StreamMemory.FramesFieldPerMemoryBlock * Stats.MemoryBlockSentCounter;
-				WaitForNextRound(StreamData.SendTimeNs);
+				WaitForNextRound();
 			}
 
 			{
-				if (CurrentFrame.IsValid())
-				{
-					//Release last frame sent. We keep hold to avoid overwriting it as rivermax is sending it
-					CurrentFrame->Clear();
-					{
-						FScopeLock Lock(&FrameCriticalSection);
-						AvailableFrames.Add(MoveTemp(CurrentFrame));
-						CurrentFrame.Reset();
-					}
-				}
-
-				while (CurrentFrame.IsValid() == false && bIsActive)
-				{
-					CurrentFrame = GetNextFrameToSend();
-					if (!CurrentFrame.IsValid())
-					{
-						TRACE_CPUPROFILER_EVENT_SCOPE(FRivermaxOutputStream::WaitForFrame);
-						const double PreWaitTime = FPlatformTime::Seconds();
-						ReadyToSendEvent->Wait();
-						UE_LOG(LogRivermax, Verbose, TEXT("Outputstream waited %0.8f"), FPlatformTime::Seconds() - PreWaitTime);
-					}
-				}
+				TRACE_CPUPROFILER_EVENT_SCOPE(FRivermaxOutputStream::PrepareNextFrame);
+				PrepareNextFrame();
 			}
 
+			
+
 			if (bIsActive == false)
 			{
 				return;
@@ -324,11 +241,7 @@ namespace UE::RivermaxCore::Private
 				if (bIsActive)
 				{
 					Stats.TotalStrides += StreamMemory.ChunkSizeInStrides;
-
-					if ((++CurrentFrame->ChunkNumber) % StreamMemory.ChunksPerFrameField == 0)
-					{
-						StreamData.SendTimeNs += StreamData.FrameFieldTimeIntervalNs;
-					}
+					++CurrentFrame->ChunkNumber;
 				}
 
 			} while (CurrentFrame->ChunkNumber < StreamMemory.ChunksPerMemoryBlock && bIsActive);
@@ -337,6 +250,8 @@ namespace UE::RivermaxCore::Private
 			Stats.MemoryBlockSentCounter++;
 			Stats.TotalStrides += StreamMemory.ChunkSizeInStrides;
 			StreamData.bHasFrameFirstChunkBeenFetched = false;
+			UE_LOG(LogRivermax, Verbose, TEXT("Stats: BlockSent: %d. CommitImmediate: %d. CommitRetries: %d. "), Stats.MemoryBlockSentCounter, Stats.CommitImmediate, Stats.CommitRetries);
+
 		}
 	}
 
@@ -407,7 +322,6 @@ namespace UE::RivermaxCore::Private
 
 	void FRivermaxOutputStream::InitializeNextFrame(const TSharedPtr<FRivermaxOutputFrame>& NextFrame)
 	{
-		NextFrame->TimestampTicks = StreamData.InitialTimestampTick;
 		NextFrame->LineNumber = 0;
 		NextFrame->PacketCounter = 0;
 		NextFrame->SRDOffset = 0;
@@ -491,7 +405,9 @@ namespace UE::RivermaxCore::Private
 		OutHeader.RawHeader[1] = 96; //todo payload type from sdp file
 		OutHeader.RawHeader[2] = (StreamData.SequenceNumber >> 8) & 0xFF;  // sequence number
 		OutHeader.RawHeader[3] = (StreamData.SequenceNumber) & 0xFF;  // sequence number
-		*(uint32*)&OutHeader.RawHeader[4] = ByteSwap((uint32)CurrentFrame->TimestampTicks);
+
+		const uint32 MediaTimestamp = GetTimestampFromTime(StreamData.NextAlignmentPointNanosec, MediaClockSampleRate);
+		*(uint32*)&OutHeader.RawHeader[4] = ByteSwap(MediaTimestamp);
 
 		//2110 specific header
 		*(uint32*)&OutHeader.RawHeader[8] = 0x0eb51dbd;  // simulated ssrc
@@ -533,7 +449,7 @@ namespace UE::RivermaxCore::Private
 			//	send_data.m_second_field = !send_data.m_second_field;
 			//	ticks /= 2;
 			//}
-			CurrentFrame->TimestampTicks += ticks;
+			//CurrentFrame->TimestampTicks += ticks;
 		}
 	}
 
@@ -555,34 +471,122 @@ namespace UE::RivermaxCore::Private
 		} while (Status == RMAX_ERR_BUSY);
 	}
 
-	void FRivermaxOutputStream::WaitForNextRound(double NextRoundTimeNs)
+	void FRivermaxOutputStream::WaitForNextRound()
 	{
-		double NowTimeNs = FPlatformTime::Seconds() * 1E9;
+		IRivermaxCoreModule& RivermaxModule = FModuleManager::LoadModuleChecked<IRivermaxCoreModule>(TEXT("RivermaxCore"));
+		const uint64 CurrentTimeNanosec = RivermaxModule.GetRivermaxManager()->GetTime();
+		const double CurrentPlatformTime = FPlatformTime::Seconds();
+		const uint64 CurrentFrameNumber = UE::RivermaxCore::GetFrameNumber(CurrentTimeNanosec, Options.FrameRate);
 
-		//If we are already late, no need to wait
-		if (NextRoundTimeNs <= NowTimeNs)
+		// Frame number we will want to align with
+		uint64 NextFrameNumber = CurrentFrameNumber;
+
+		bool bHasValidTimings = true;
+		if (StreamData.bHasValidNextFrameNumber == false)
 		{
-			return;
+			// Number here is quite big to patch an issue found during testing.
+			// Sometimes, first wait is not long enough and we end up scheduling 2 frames super close. 
+			// Which causes rivermax hardware send queues to fill up and cause retries
+			// which increases the time it takes to commit all chunks per frame
+			// Will need to be improved. 
+			NextFrameNumber = CurrentFrameNumber + 10;
 		}
+		else
+		{
+			// Case where we are back and frame number is the previous one. Depending on offsets, this could happen
+			if (CurrentFrameNumber == StreamData.NextAlignmentPointFrameNumber - 1)
+			{
+				NextFrameNumber = StreamData.NextAlignmentPointFrameNumber + 1;
+				UE_LOG(LogRivermax, Verbose, TEXT("Scheduling last frame was faster than expected. (CurrentFrame: '%llu' LastScheduled: '%llu') Scheduling for following expected one.")
+				, CurrentFrameNumber
+				, StreamData.NextAlignmentPointFrameNumber);
+			}
+			else
+			{
+				// We expect current frame number to be the one we scheduled for the last time or greater if something happened
+				if (CurrentFrameNumber >= StreamData.NextAlignmentPointFrameNumber)
+				{
+					// If current frame is greater than last scheduled, we missed an alignment point. Shouldn't happen with continuous thread independent of engine.
+					const uint64 DeltaFrames = CurrentFrameNumber - StreamData.NextAlignmentPointFrameNumber;
+					if (DeltaFrames >= 1)
+					{
+						UE_LOG(LogRivermax, Warning, TEXT("Output missed %llu frames."), DeltaFrames);
+						// For now, schedule for the following frame as normal but might need to revisit this behavior if it causes issues.
+					}
 
-		const double SleepTimeNs = NextRoundTimeNs - NowTimeNs;
+					NextFrameNumber = CurrentFrameNumber + 1;
+				}
+				else
+				{
+					// This is not expected (going back in time) but we should be able to continue. Scheduling immediately
+					ensure(false);
+					bHasValidTimings = false;
+				}
+			}
+		}
+
+		// Get next alignment point based on the frame number we are aligning with
+		const uint64 NextAlignmentNano = UE::RivermaxCore::GetAlignmentPointFromFrameNumber(NextFrameNumber, Options.FrameRate);
 
-		static const bool bShouldPoll = false;
+		// Add Tro offset to next alignment point
+		StreamData.NextAlignmentPointNanosec = NextAlignmentNano;
+		StreamData.NextScheduleTimeNanosec = NextAlignmentNano + TransmitOffsetNanosec;
+		StreamData.NextAlignmentPointFrameNumber = NextFrameNumber;
+
+		// Offset wakeup if desired to give more time for scheduling. 
+		const uint64 WakeupTime = NextAlignmentNano - AlignmentPointSchedulingLeadTimeNanosec;
+
+		// Used to know if we schedule blindly in the future or based on previous data
+		StreamData.bHasValidNextFrameNumber = bHasValidTimings;
+		
 
-		//Todo Should we have a threshold where we don't sleep?
-		//Todo Should we sleep less when we had retries
-		if (bShouldPoll)
+		uint64 WaitTimeNanosec = WakeupTime - CurrentTimeNanosec;
+
+		// Wakeup can be smaller than current time with controllable offset
+		if (WakeupTime < CurrentTimeNanosec)
+		{
+			WaitTimeNanosec = 0;
+		}
+
+		static constexpr float SleepThresholdSec = 5.0f * (1.0f / 1000.0f);
+		static constexpr float YieldTimeSec = 2.0f * (1.0f / 1000.0f);
+		double WaitTimeSec = WaitTimeNanosec / 1E9;
+
+		// Protect for erroneously long wait time
+		if (!ensure(WaitTimeSec < 1.0))
+		{
+			WaitTimeSec = 1.0;
+		}
+
+		// Sleep for the largest chunk of time
+		if (WaitTimeSec > SleepThresholdSec)
+		{
+			FPlatformProcess::SleepNoStats(WaitTimeSec - YieldTimeSec);
+		}
+
+		// Yield our thread time for the remainder of wait time. 
+		// Should we spin for a smaller time to be more precise?
+		// Should we use platform time instead of rivermax get PTP to avoid making calls to it?
+		constexpr bool bUsePlatformTimeToSpin = true;
+		if (bUsePlatformTimeToSpin == false)
 		{
-			while (NowTimeNs < SleepTimeNs)
+			while (RivermaxModule.GetRivermaxManager()->GetTime() < WakeupTime)
 			{
-				NowTimeNs = FPlatformTime::Seconds() * 1E9;
+				FPlatformProcess::SleepNoStats(0.f);
 			}
 		}
 		else
 		{
-			const double SleepTimeS = SleepTimeNs / 1E9;
-			UE_LOG(LogRivermax, Verbose, TEXT("Outputstream sleeping for %0.5f"), SleepTimeS);
-			FPlatformProcess::SleepNoStats(SleepTimeS);
+			while (FPlatformTime::Seconds() < (CurrentPlatformTime + WaitTimeSec))
+			{
+				FPlatformProcess::SleepNoStats(0.f);
+			}
+		}
+
+		if (StreamData.bHasValidNextFrameNumber)
+		{
+			const uint64 AfterSleepTimeNanosec = RivermaxModule.GetRivermaxManager()->GetTime();
+			UE_LOG(LogRivermax, Verbose, TEXT("Scheduling at %llu. CurrentTime %llu. NextAlign %llu. Waiting %0.9f"), StreamData.NextScheduleTimeNanosec, CurrentTimeNanosec, NextAlignmentNano, (double)WaitTimeNanosec / 1E9);
 		}
 	}
 
@@ -644,9 +648,22 @@ namespace UE::RivermaxCore::Private
 		{
 			// todo configure timeout
 			
-			const uint64 Timeout = 0;
+			//Only first chunk gets scheduled with a timestamp. Following chunks are queued after it using 0
+			uint64 ScheduleTime = CurrentFrame->ChunkNumber == 0 ? StreamData.NextScheduleTimeNanosec : 0;
 			const rmax_commit_flags_t CommitFlags{};
-			Status = rmax_out_commit(StreamId, Timeout, CommitFlags);
+			if (ScheduleTime != 0)
+			{
+				IRivermaxCoreModule& RivermaxModule = FModuleManager::LoadModuleChecked<IRivermaxCoreModule>(TEXT("RivermaxCore"));
+				const uint64 CurrentTimeNanosec = RivermaxModule.GetRivermaxManager()->GetTime();
+				if (ScheduleTime <= CurrentTimeNanosec)
+				{
+					ScheduleTime = 0;
+					++Stats.CommitImmediate;
+				}
+			}
+
+			Status = rmax_out_commit(StreamId, ScheduleTime, CommitFlags);
+
 			if (Status == RMAX_OK)
 			{
 				break;
@@ -661,6 +678,12 @@ namespace UE::RivermaxCore::Private
 				Listener->OnStreamError();
 				bIsActive = false;
 			}
+			else
+			{
+				UE_LOG(LogRivermax, Error, TEXT("Unhandled error (%d) while trying to commit next round of chunks."), Status);
+				Listener->OnStreamError();
+				bIsActive = false;
+			}
 
 		} while (Status != RMAX_OK && bIsActive);
 	}
@@ -692,5 +715,142 @@ namespace UE::RivermaxCore::Private
 	{
 
 	}
+
+	void FRivermaxOutputStream::PrepareNextFrame()
+	{
+		// We always want to send out a frame at the desired interval. If a new frame is not ready, reuse last one. 
+		bool bHasFrameToSend = false;
+		{
+			FScopeLock Lock(&FrameCriticalSection);
+			bHasFrameToSend = FramesToSend.IsEmpty() == false;
+		}
+
+		// First frame occurrence when none were produced yet so we wait for the first one to be available
+		if (CurrentFrame.IsValid() == false && bHasFrameToSend == false)
+		{
+			//First iteration before first frame available
+			ReadyToSendEvent->Wait();
+		}
+
+		// Here, we either have a frame available to send or the last one to resend
+		{
+			{
+				FScopeLock Lock(&FrameCriticalSection);
+				bHasFrameToSend = FramesToSend.IsEmpty() == false;
+			}
+
+			check(bHasFrameToSend || CurrentFrame.IsValid());
+
+			if (bHasFrameToSend)
+			{
+				if (CurrentFrame.IsValid())
+				{
+					//Release last frame sent. We keep hold to avoid overwriting it as rivermax is sending it
+					CurrentFrame->Reset();
+					{
+						FScopeLock Lock(&FrameCriticalSection);
+						AvailableFrames.Add(MoveTemp(CurrentFrame));
+						CurrentFrame.Reset();
+					}
+				}
+
+				CurrentFrame = GetNextFrameToSend();
+			}
+			else
+			{
+				// No frame to send, keep last one and restarts its internal counters
+				UE_LOG(LogRivermax, Warning, TEXT("No frame to send. Reusing last frame '%d'"), CurrentFrame->FrameIndex);
+				InitializeNextFrame(CurrentFrame);		
+				
+				// Since we want to resend last frame, we need to fast forward chunk pointer to re-point to the one we just sent
+				rmax_status_t Status = rmax_out_skip_chunks(StreamId, StreamMemory.ChunksPerFrameField * (Options.NumberOfBuffers - 1));
+				if (Status != RMAX_OK)
+				{
+					// If we don't go in error, we end up with misaligned chunk vs frame pointer to write. 
+					ensure(false);
+					UE_LOG(LogRivermax, Error, TEXT("Invalid error happened while trying to skip chunks. Status: %d."), Status);
+					Listener->OnStreamError();
+					bIsActive = false;
+				}
+			}
+		}
+
+		//TRACE_BOOKMARK(TEXT("Schedulin %llu"), GetTimestampFromTime(StreamData.NextAlignmentPointNanosec, MediaClockSampleRate));
+	}
+
+	void FRivermaxOutputStream::InitializeStreamTimingSettings()
+	{
+		using namespace UE::RivermaxCore::Private::Utils;
+
+		double FrameIntervalNs = StreamData.FrameFieldTimeIntervalNs;
+		const bool bIsProgressive = true;//todo MediaConfiguration.IsProgressive() 
+		uint32 PacketsInFrameField = StreamMemory.PacketsInFrameField;
+		if (bIsProgressive == false)
+		{
+			FrameIntervalNs *= 2;
+			PacketsInFrameField *= 2;
+		}
+
+		// See https://ieeexplore.ieee.org/stamp/stamp.jsp?tp=&arnumber=8165971 for reference
+		// Gapped PRS doesn't support non standard resolution. Linear PRS would but Rivermax doesn't support it.
+		if (StreamType == ERivermaxStreamType::VIDEO_2110_20_STREAM)
+		{
+			double RActive;
+			double TRODefaultMultiplier;
+			if (bIsProgressive)
+			{
+				RActive = (1080.0 / 1125.0);
+				if (Options.Resolution.Y >= FullHDHeight)
+				{
+					// As defined by SMPTE 2110-21 6.3.2
+					TRODefaultMultiplier = (43.0 / 1125.0);
+				}
+				else
+				{
+					TRODefaultMultiplier = (28.0 / 750.0);
+				}
+			}
+			else
+			{
+				if (Options.Resolution.Y >= FullHDHeight)
+				{
+					// As defined by SMPTE 2110-21 6.3.3
+					RActive = (1080.0 / 1125.0);
+					TRODefaultMultiplier = (22.0 / 1125.0);
+				}
+				else if (Options.Resolution.Y >= 576)
+				{
+					RActive = (576.0 / 625.0);
+					TRODefaultMultiplier = (26.0 / 625.0);
+				}
+				else
+				{
+					RActive = (487.0 / 525.0);
+					TRODefaultMultiplier = (20.0 / 525.0);
+				}
+			}
+
+			// Need to reinvestigate the implication of this and possibly add cvar to control it runtime
+			const double TRSNano = (FrameIntervalNs * RActive) / PacketsInFrameField;
+			TransmitOffsetNanosec = (uint64)((TRODefaultMultiplier * FrameIntervalNs) - (VideoTROModification * TRSNano));
+		}
+	}
+
+	uint32 FRivermaxOutputStream::GetTimestampFromTime(uint64 InTimeNanosec, double InMediaClockRate) const
+	{
+		// RTP timestamp is 32 bits and based on media clock (usually 90kHz).
+		// Conversion based on rivermax samples
+
+		const uint64 Nanoscale = 1E9;
+		const uint64 Seconds = InTimeNanosec / Nanoscale;
+		const uint64 Nanoseconds = InTimeNanosec % Nanoscale;
+		const uint64 MediaFrameNumber = Seconds * InMediaClockRate;
+		const uint64 MediaSubFrameNumber = Nanoseconds * InMediaClockRate / Nanoscale;
+		const double Mask = 0x100000000;
+		const double MediaTime = FMath::Fmod(MediaFrameNumber, Mask);
+		const double MediaSubTime = FMath::Fmod(MediaSubFrameNumber, Mask);
+		return MediaTime + MediaSubTime;
+	}
+
 }
 
diff --git a/Engine/Plugins/VirtualProduction/Rivermax/RivermaxCore/Source/RivermaxCore/Private/RivermaxOutputStream.h b/Engine/Plugins/VirtualProduction/Rivermax/RivermaxCore/Source/RivermaxCore/Private/RivermaxOutputStream.h
--- a/Engine/Plugins/VirtualProduction/Rivermax/RivermaxCore/Source/RivermaxCore/Private/RivermaxOutputStream.h
+++ b/Engine/Plugins/VirtualProduction/Rivermax/RivermaxCore/Source/RivermaxCore/Private/RivermaxOutputStream.h
@@ -56,7 +56,8 @@ namespace UE::RivermaxCore::Private
 		uint32 ChunkWait = 0;
 		uint32 CommitWaits = 0;
 		uint32 CommitRetries = 0;
-		uint64 MemoryBlockSentCounter = 0; //Global for an active capture session to track timestamp for next packet
+		uint32 CommitImmediate = 0;
+		uint64 MemoryBlockSentCounter = 0;
 	};
 
 	struct FRivermaxOutputStreamData
@@ -64,10 +65,21 @@ namespace UE::RivermaxCore::Private
 		/** Current sequence number being done */
 		uint32 SequenceNumber = 0;
 		double FrameFieldTimeIntervalNs = 0.0;
-		double StartSendTimeNs = 0.0;
-		double SendTimeNs = 0.0;
-		double InitialTimestampTick = 0.0;
+
+		/** Used to detect misalignment between chunk being sent and frame memory we are writing in */
 		bool bHasFrameFirstChunkBeenFetched = false;
+
+		/** Next alignment point based on PTP standard */
+		uint64 NextAlignmentPointNanosec = 0;
+
+		/** Next schedule time using 2110 gapped model timing and controllable offset */
+		uint64 NextScheduleTimeNanosec = 0;
+
+		/** Whether next alignment frame number is deemed valid or not to detect missed frames. */
+		bool bHasValidNextFrameNumber = false;
+		
+		/** Next alignment point frame number treated to detect missed frames */
+		uint64 NextAlignmentPointFrameNumber = 0;
 	};
 
 
@@ -102,10 +114,13 @@ namespace UE::RivermaxCore::Private
 		TSharedPtr<FRivermaxOutputFrame> GetNextAvailableFrame(uint32 InFrameIdentifier);
 		void BuildRTPHeader(FRTPHeader& OutHeader) const;
 		void DestroyStream();
-		void WaitForNextRound(double NextRoundTime);
+		void WaitForNextRound();
 		void GetNextChunk();
 		void SetupRTPHeaders();
 		void CommitNextChunks();
+		void PrepareNextFrame();
+		void InitializeStreamTimingSettings();
+		uint32 GetTimestampFromTime(uint64 InTimeNanosec, double InMediaClockRate) const;
 
 	private:
 		FRivermaxStreamOptions Options;
@@ -120,7 +135,6 @@ namespace UE::RivermaxCore::Private
 
 		TArray<TSharedPtr<FRivermaxOutputFrame>> AvailableFrames;
 		TArray<TSharedPtr<FRivermaxOutputFrame>> FramesToSend;
-		double LastTime = 0.0;
 
 		TUniquePtr<FRunnableThread> RivermaxThread;
 		std::atomic<bool> bIsActive;
@@ -131,6 +145,12 @@ namespace UE::RivermaxCore::Private
 
 		static constexpr double MediaClockSampleRate = 90000.0; //Required to comply with SMTPE 2110-10.The Media Clock and RTP Clock rate for streams compliant to this standard shall be 90 kHz.
 		ERivermaxStreamType StreamType = ERivermaxStreamType::VIDEO_2110_20_STREAM; //todo
+
+		/** TRoffset time calculated based on ST2110 - 21 Gapped(for now) method. This is added to next alignment point */
+		uint64 TransmitOffsetNanosec = 0;
+
+		/** Offset substracted to next alignment point (wake up) to move the start of scheduling */
+		uint64 AlignmentPointSchedulingLeadTimeNanosec = 0;
 	};
 }
 
diff --git a/Engine/Plugins/VirtualProduction/Rivermax/RivermaxCore/Source/RivermaxCore/Private/RivermaxPTPUtils.cpp b/Engine/Plugins/VirtualProduction/Rivermax/RivermaxCore/Source/RivermaxCore/Private/RivermaxPTPUtils.cpp
new file mode 100644
--- /dev/null
+++ b/Engine/Plugins/VirtualProduction/Rivermax/RivermaxCore/Source/RivermaxCore/Private/RivermaxPTPUtils.cpp
@@ -0,0 +1,122 @@
+// Copyright Epic Games, Inc. All Rights Reserved.
+
+#include "RivermaxPTPUtils.h"
+
+#include "Misc/FrameRate.h"
+#include "RivermaxTypes.h"
+
+namespace UE::RivermaxCore
+{
+	uint64 GetNextAlignmentPoint(const uint64 InPTPTimeNanosec, const FFrameRate& InRate)
+	{
+		// Based on ST2059. Formulae from PTP time to next alignment point
+		// NextAlignmentPoint = Floor( (TimeNs / IntervalNs + 1) * IntervalNs)
+		// Time involved don't play well with double. Need to keep fixed point as much as possible
+
+		// First part of formulae is to get the associated frame number of current time
+		const uint64 CurrentFrameNumber = GetFrameNumber(InPTPTimeNanosec, InRate);
+
+		// We want to get the time of the following frame number
+		return GetAlignmentPointFromFrameNumber(CurrentFrameNumber + 1, InRate);
+	}
+
+	uint64 GetAlignmentPointFromFrameNumber(const uint64 InFrameNumber, const FFrameRate& InRate)
+	{
+		// Based on ST2059. Formulae from PTP time to next alignment point
+		// NextAlignmentPoint = Floor( (CurrentFrameNumber + 1) * IntervalNs)
+
+		// Example for frame F=40113798005 at 23.97
+		// uint64 T = F / 24000 = 1671408
+		// uint64 Tr =F % 24000 = 6005
+		// uint64 TT = T * 1001 = 1673079408
+		// uint64 TTr = 6005 * 1001 = 6011005
+		// uint64 TTrr = TTr * Scale / 24000 = 250458541666
+		// return TT * Scale + TTrr = 1673079658458541666
+
+		const uint64 Nanoscale = 1E9;
+
+		const uint64 Intermediate = InFrameNumber / InRate.Numerator;
+		const uint64 IntermediateRemainder = InFrameNumber % InRate.Numerator;
+		const uint64 NextAlignment = Intermediate * InRate.Denominator;
+		const uint64 NextAlignementRemainderInterm = (IntermediateRemainder * InRate.Denominator);
+		const uint64 NextAlignmentRemainder = (NextAlignementRemainderInterm * Nanoscale) / InRate.Numerator;
+
+		return (NextAlignment * Nanoscale) + NextAlignmentRemainder;
+	}
+
+	uint64 GetFrameNumber(const uint64 InPTPTimeNanosec, const FFrameRate& InRate)
+	{
+		const uint64 NanoScale = 1E9;
+
+		// Util struct returned for intermediate calculation of frame number
+		struct FResult
+		{
+			uint64 FrameCount = 0;
+			double Fraction = 0.0;
+		};
+
+		const auto CalculateFrameNumberFunc = [InRate](const uint64 PTPTime, uint64 Scale) -> FResult
+		{
+			// Example for PTP time 1673079658458541667 at 23.97
+			
+			// Seconds part (1673079658). Scale = 1. Remainder = 0.
+			// uint64 F = Time / 1001 = 1671408
+			// uint64 Fr = Time % 1001 = 250
+			// uint64 FF = F * 24000 = 40113792000
+			// uint64 FFr = (Fr * 24000) / 1001 = 5994
+			// uint64 FFrr =(Fr * 24000) % 1001) = 6
+			// double Tr = (FFrr / (double)1001.0 = 0.005994005994
+			// uint64 Frame = (FF + FFr) / Scale = 40113797994
+			// double FractionalFrame = Frac((FF + FFr) / Scale) + Tr / Scale = 0.005994005994
+			// ret {Frame, Tr} = {40113797994, 0.005994005994}
+
+			// Nanoseconds part (458541667). Scale = 1E9. Remainder = 0.005994005994 * 1E9 = 5994005
+			// uint64 F = Time / 1001 = 458083
+			// uint64 Fr = Time % 1001 = 584
+			// uint64 FF = F * 24000 = 10993992000
+			// uint64 FFr = (Fr * 24000) / 1001 = 14001
+			// uint64 FFrr =(Fr * 24000) % 1001) = 999
+			// double Tr = (FFrr / (double)1001.0 = 0.998001998001998002
+			// uint64 Frame = (FF + FFr) / Scale = 10
+			// ret {Frame, Tr} = {10, .994006001998002}
+
+			// Full frames = 40113797994 + 10 =  40113798004
+			// Fractional frames = 1.0000000079920019 ( 0.005994005994 + 0.994006001998002 )
+			// We round to the 1/1000000th of a frame to cover cases where fractional frame is 0.999999
+			// Final result= 40113798004 + 1 = 40113798005
+
+
+			const uint64 FrameIntermediate = PTPTime / InRate.Denominator;
+			const uint64 FrameRemainderIntermediate = PTPTime % InRate.Denominator;
+			const uint64 FrameInteger = FrameIntermediate * InRate.Numerator;
+			const uint64 FrameIntegerRemainder = (FrameRemainderIntermediate * InRate.Numerator) / InRate.Denominator;
+			const double TrailingRemainder = ((FrameRemainderIntermediate * InRate.Numerator) % InRate.Denominator) / (double)InRate.Denominator;
+			const uint64 FrameNumber = (FrameInteger + FrameIntegerRemainder) / Scale;
+
+			const double FractionalFrame  = FMath::Frac((FrameInteger + FrameIntegerRemainder) / (double)Scale) + TrailingRemainder / Scale;
+			return FResult{ FrameNumber, FractionalFrame };
+		};
+
+		// Split incoming time in nanoseconds in two parts : seconds and nanoseconds
+		// Even after dividing by frame rate denominator (i.e 1001), we might overflow uint64 when multiplying by frame rate numerator (i.e. 24000)
+		const uint64 PTPSeconds = InPTPTimeNanosec / NanoScale;
+		const uint64 PTPNanoSeconds = InPTPTimeNanosec % NanoScale;
+		
+		// Proceed with seconds part. No scaling
+		const FResult FrameNumber1 = CalculateFrameNumberFunc(PTPSeconds, 1);
+
+		// Proceed with nanoseconds using a scale of 1E9 to get actual frame numbers
+		const FResult FrameNumber2 = CalculateFrameNumberFunc(PTPNanoSeconds, NanoScale);
+
+		// Add up full frame number
+		const uint64 FrameNumber = FrameNumber1.FrameCount + FrameNumber2.FrameCount;
+
+		// Add up fractional frame number from intermediate calculation
+		// and round to the 1/100000th of a frame to avoid returning the previous time when the incoming time
+		// is almost exactly on alignment
+		const double FramePrecision = 100000; 
+		const double FractionFrameNumber = FrameNumber1.Fraction + FrameNumber2.Fraction;
+		const uint64 TrailingFrameNumber = FMath::RoundToDouble(FractionFrameNumber * FramePrecision) / FramePrecision;
+		return FrameNumber + TrailingFrameNumber;
+	}
+}
\ No newline at end of file
diff --git a/Engine/Plugins/VirtualProduction/Rivermax/RivermaxCore/Source/RivermaxCore/Private/Tests/PTPTest.cpp b/Engine/Plugins/VirtualProduction/Rivermax/RivermaxCore/Source/RivermaxCore/Private/Tests/PTPTest.cpp
new file mode 100644
--- /dev/null
+++ b/Engine/Plugins/VirtualProduction/Rivermax/RivermaxCore/Source/RivermaxCore/Private/Tests/PTPTest.cpp
@@ -0,0 +1,106 @@
+// Copyright Epic Games, Inc. All Rights Reserved.
+
+#include "CoreTypes.h"
+#include "Async/Async.h"
+#include "Async/Future.h"
+#include "Containers/UnrealString.h"
+#include "Logging/LogMacros.h"
+#include "Math/NumericLimits.h"
+#include "Misc/FrameRate.h"
+#include "Misc/Timecode.h"
+#include "Misc/Timespan.h"
+#include "Misc/AutomationTest.h"
+#include "RivermaxPTPUtils.h"
+
+PRAGMA_DISABLE_OPTIMIZATION
+#if WITH_DEV_AUTOMATION_TESTS
+
+IMPLEMENT_SIMPLE_AUTOMATION_TEST(FPTPTest, "Plugin.Rivermax.PTP", EAutomationTestFlags::ApplicationContextMask | EAutomationTestFlags::EngineFilter)
+
+struct FTestData
+{
+	FFrameRate Rate;
+	uint64 TimeNanosec = 0;
+	uint64 ExpectedFrameNumber = 0;
+	uint64 ExpectedAlignmentPointNanosec = 0;
+};
+
+bool FPTPTest::RunTest(const FString& Parameters)
+{
+	TArray<FTestData> TestPoints;
+
+	// Add couple of test points at different frame rate to test PTP operation
+
+	TestPoints.Add({ FFrameRate(24,1), 0, 0, 41666666 });
+	TestPoints.Add({ FFrameRate(24,1), 42000000, 1, 83333333 });
+	TestPoints.Add({ FFrameRate(24,1), 1671408000000000000, 40113792000, 1671408000041666666 });
+	TestPoints.Add({ FFrameRate(24,1), 1671408000041666666, 40113792001, 1671408000083333333 }); 
+	TestPoints.Add({ FFrameRate(24,1), 1671408000083333333, 40113792002, 1671408000125000000 });
+	TestPoints.Add({ FFrameRate(24,1), 1671408000020000000, 40113792000, 1671408000041666666 }); 
+	TestPoints.Add({ FFrameRate(24,1), 1671408000040000000, 40113792000, 1671408000041666666 });
+	TestPoints.Add({ FFrameRate(24,1), 1673079658458541667, 40153911803, 1673079658500000000 });
+	TestPoints.Add({ FFrameRate(24,1), 4099680000000000000, 98392320000, 4099680000041666666 });
+	TestPoints.Add({ FFrameRate(24,1), 4099680000041666666, 98392320001, 4099680000083333333 }); // Large nano second test. Approx 130 years from epoch
+
+	TestPoints.Add({ FFrameRate(24000,1001), 0, 0, 41708333 }); 
+	TestPoints.Add({ FFrameRate(24000,1001), 42000000, 1, 83416666 });
+	TestPoints.Add({ FFrameRate(24000,1001), 1673079408000000000, 40113792000, 1673079408041708333 });
+	TestPoints.Add({ FFrameRate(24000,1001), 1673079408020000000, 40113792000, 1673079408041708333 });  
+	TestPoints.Add({ FFrameRate(24000,1001), 1673079408040000000, 40113792000, 1673079408041708333 });  
+	TestPoints.Add({ FFrameRate(24000,1001), 1673079408041708333, 40113792001, 1673079408083416666 }); // Right on alignment. Without frame number round up, it would return same input time
+	TestPoints.Add({ FFrameRate(24000,1001), 1673079408083416666, 40113792002, 1673079408125125000 }); // Right on alignment.
+	TestPoints.Add({ FFrameRate(24000,1001), 1673079408125125000, 40113792003, 1673079408166833333 }); // Right on alignment.
+	TestPoints.Add({ FFrameRate(24000,1001), 1673079408125124999, 40113792003, 1673079408166833333 }); // Right on alignment.
+	TestPoints.Add({ FFrameRate(24000,1001), 1673079408166833333, 40113792004, 1673079408208541666 });
+	TestPoints.Add({ FFrameRate(24000,1001), 1673079408208541666, 40113792005, 1673079408250250000 });
+	TestPoints.Add({ FFrameRate(24000,1001), 1673079408250249998, 40113792006, 1673079408291958333 });
+	TestPoints.Add({ FFrameRate(24000,1001), 1673079408250250000, 40113792006, 1673079408291958333 });
+	TestPoints.Add({ FFrameRate(24000,1001), 1673079408291958333, 40113792007, 1673079408333666666 });
+	TestPoints.Add({ FFrameRate(24000,1001), 1673079408333666666, 40113792008, 1673079408375375000 });
+	TestPoints.Add({ FFrameRate(24000,1001), 4099680000000000000, 98294025974, 4099680000040625000 });
+	TestPoints.Add({ FFrameRate(24000,1001), 4099680000040625000, 98294025975, 4099680000082333333 });
+	TestPoints.Add({ FFrameRate(24000,1001), 4099680000082333333, 98294025976, 4099680000124041666 });
+
+	TestPoints.Add({ FFrameRate(60,1), 0, 0, 16666666 });
+	TestPoints.Add({ FFrameRate(60,1), 26666666, 1, 33333333 });
+	TestPoints.Add({ FFrameRate(60,1), 2674252800000000000, 160455168000, 2674252800016666666 }); 
+	TestPoints.Add({ FFrameRate(60,1), 2674252800008000000, 160455168000, 2674252800016666666 });
+	TestPoints.Add({ FFrameRate(60,1), 2674252800010000000, 160455168000, 2674252800016666666 });
+	TestPoints.Add({ FFrameRate(60,1), 1673079658458541667, 100384779507, 1673079658466666666 });
+	TestPoints.Add({ FFrameRate(60,1), 1673079658466666666, 100384779508, 1673079658483333333 });
+	TestPoints.Add({ FFrameRate(60,1), 4099680000000000000, 245980800000, 4099680000016666666 });
+	TestPoints.Add({ FFrameRate(60,1), 4099680000016666666, 245980800001, 4099680000033333333 });
+
+	TestPoints.Add({ FFrameRate(60000,1001), 0, 0, 16683333 });
+	TestPoints.Add({ FFrameRate(60000,1001), 26666666, 1, 33366666 });
+	TestPoints.Add({ FFrameRate(60000,1001), 2674252800000000000, 160294873126, 2674252800002116666 });
+	TestPoints.Add({ FFrameRate(60000,1001), 2674252800008000000, 160294873127, 2674252800018800000 });
+	TestPoints.Add({ FFrameRate(60000,1001), 2674252800010000000, 160294873127, 2674252800018800000 });
+	TestPoints.Add({ FFrameRate(60000,1001), 1673079658458541667, 100284495012, 1673079658466883333 });
+	TestPoints.Add({ FFrameRate(60000,1001), 1673079658466883333, 100284495013, 1673079658483566666 });
+	TestPoints.Add({ FFrameRate(60000,1001), 4099680000000000000, 245735064935, 4099680000015600000 });
+	
+
+	const auto AlignmentTestFunc = [this](const FTestData& InTestPoint)
+	{
+		const uint64 FoundFrameNumber = UE::RivermaxCore::GetFrameNumber(InTestPoint.TimeNanosec, InTestPoint.Rate);
+		const FString FrameNumberError = FString::Printf(TEXT("Frame rate %s: Error calculating frame number for time %llu."), *InTestPoint.Rate.ToPrettyText().ToString(), InTestPoint.TimeNanosec);
+		const FString AlignmentPointError = FString::Printf(TEXT("Frame rate %s: Error calculating next alignment point for time %llu."), *InTestPoint.Rate.ToPrettyText().ToString(), InTestPoint.TimeNanosec);
+		TestEqual(*FrameNumberError, FoundFrameNumber, InTestPoint.ExpectedFrameNumber);
+		const uint64 NextAlignmentPointNanosec = UE::RivermaxCore::GetNextAlignmentPoint(InTestPoint.TimeNanosec, InTestPoint.Rate);
+		TestEqual(*AlignmentPointError, NextAlignmentPointNanosec, InTestPoint.ExpectedAlignmentPointNanosec);
+
+		return true;
+	};
+
+	for (const FTestData& TestPoint : TestPoints)
+	{
+		AlignmentTestFunc(TestPoint);
+	}
+
+	return true;
+}
+
+
+#endif //WITH_DEV_AUTOMATION_TESTS
+PRAGMA_ENABLE_OPTIMIZATION
\ No newline at end of file
diff --git a/Engine/Plugins/VirtualProduction/Rivermax/RivermaxCore/Source/RivermaxCore/Public/RivermaxPTPUtils.h b/Engine/Plugins/VirtualProduction/Rivermax/RivermaxCore/Source/RivermaxCore/Public/RivermaxPTPUtils.h
new file mode 100644
--- /dev/null
+++ b/Engine/Plugins/VirtualProduction/Rivermax/RivermaxCore/Source/RivermaxCore/Public/RivermaxPTPUtils.h
@@ -0,0 +1,29 @@
+// Copyright Epic Games, Inc. All Rights Reserved.
+
+#pragma once
+
+#include "CoreMinimal.h"
+
+
+struct FFrameRate;
+
+namespace UE::RivermaxCore
+{
+	/** 
+	 * Returns next alignment point in nanosecond based on incoming PTP time and frame rate.
+	 * Follows ST2059 standard for alignment
+	 */
+	RIVERMAXCORE_API uint64 GetNextAlignmentPoint(const uint64 InPTPTimeNanosec, const FFrameRate& InRate);
+	
+	/**
+	 * Returns alignment point in nanosecond based on incoming frame number and frame rate.
+	 */
+	RIVERMAXCORE_API uint64 GetAlignmentPointFromFrameNumber(const uint64 InFrameNumber, const FFrameRate& InRate);
+
+	/**
+	 * Returns current frame number for the incoming PTP time for the given frame rate.
+	 */
+	RIVERMAXCORE_API uint64 GetFrameNumber(const uint64 InPTPTimeNanosec, const FFrameRate& InRate);
+}
+
+
diff --git a/Engine/Plugins/VirtualProduction/Rivermax/RivermaxMedia/Source/RivermaxMedia/Private/RivermaxCustomTimeStep.cpp b/Engine/Plugins/VirtualProduction/Rivermax/RivermaxMedia/Source/RivermaxMedia/Private/RivermaxCustomTimeStep.cpp
new file mode 100644
--- /dev/null
+++ b/Engine/Plugins/VirtualProduction/Rivermax/RivermaxMedia/Source/RivermaxMedia/Private/RivermaxCustomTimeStep.cpp
@@ -0,0 +1,220 @@
+// Copyright Epic Games, Inc. All Rights Reserved.
+
+#include "RivermaxCustomTimeStep.h"
+
+#include "IRivermaxCoreModule.h"
+#include "IRivermaxManager.h"
+#include "RivermaxMediaLog.h"
+#include "RivermaxPTPUtils.h"
+#include "Stats/StatsMisc.h"
+
+
+#define LOCTEXT_NAMESPACE "RivermaxCustomTimeStep"
+
+
+
+bool URivermaxCustomTimeStep::Initialize(UEngine* InEngine)
+{
+	IRivermaxCoreModule* RivermaxModule = FModuleManager::GetModulePtr<IRivermaxCoreModule>("RivermaxCore");
+	if (RivermaxModule == nullptr)
+	{
+		State = ECustomTimeStepSynchronizationState::Error;
+		return false;
+	}
+
+	if (RivermaxModule->IsAvailable() == false)
+	{
+		State = ECustomTimeStepSynchronizationState::Error;
+		return false;
+	}
+
+	if (RivermaxModule->GetRivermaxManager()->GetDevices().Num() <= 0)
+	{
+		State = ECustomTimeStepSynchronizationState::Error;
+		return false;
+	}
+
+	if (RivermaxModule->GetRivermaxManager()->GetTimeSource() != ERivermaxTimeSource::PTP)
+	{
+		State = ECustomTimeStepSynchronizationState::Error;
+		return false;
+	}
+
+	State = ECustomTimeStepSynchronizationState::Synchronized;
+	return true;
+}
+
+
+void URivermaxCustomTimeStep::Shutdown(UEngine* InEngine)
+{
+#if WITH_EDITORONLY_DATA
+	InitializedEngine = nullptr;
+#endif
+
+	State = ECustomTimeStepSynchronizationState::Closed;
+
+}
+
+bool URivermaxCustomTimeStep::UpdateTimeStep(UEngine* InEngine)
+{
+	IRivermaxCoreModule* RivermaxModule = FModuleManager::GetModulePtr<IRivermaxCoreModule>("RivermaxCore");
+	if (RivermaxModule == nullptr)
+	{
+		State = ECustomTimeStepSynchronizationState::Error;
+	}
+
+	if (State == ECustomTimeStepSynchronizationState::Closed)
+	{
+		return true;
+	}
+
+	if (State == ECustomTimeStepSynchronizationState::Error)
+	{
+		bIsPreviousFrameNumberValid = false;
+		return true;
+	}
+
+	// Updates logical last time to match logical current time from last tick
+	UpdateApplicationLastTime();
+
+	const double TimeBeforeSync = FPlatformTime::Seconds();
+
+	const bool bWaitedForSync = WaitForSync();
+
+	const double TimeAfterSync = FPlatformTime::Seconds();
+
+	if (!bWaitedForSync)
+	{
+		State = ECustomTimeStepSynchronizationState::Error;
+		return true;
+	}
+
+	UpdateAppTimes(TimeBeforeSync, TimeAfterSync);
+
+	return false;
+}
+
+ECustomTimeStepSynchronizationState URivermaxCustomTimeStep::GetSynchronizationState() const
+{
+	return State;
+}
+
+FFrameRate URivermaxCustomTimeStep::GetFixedFrameRate() const
+{
+	return FrameRate;
+}
+
+uint32 URivermaxCustomTimeStep::GetLastSyncCountDelta() const
+{
+	return DeltaFrameNumber;
+}
+
+bool URivermaxCustomTimeStep::IsLastSyncDataValid() const
+{
+	return bIsPreviousFrameNumberValid;
+}
+
+FFrameRate URivermaxCustomTimeStep::GetSyncRate() const
+{
+	FFrameRate SyncRate = GetFixedFrameRate();
+	//Need to handle SyncRate times 2 for interlace format
+	return SyncRate;
+}
+
+bool URivermaxCustomTimeStep::WaitForSync()
+{
+	DeltaFrameNumber = 1;
+
+	const bool bWasWaitValid = WaitForNextFrame();
+	if (!bWasWaitValid)
+	{
+		bIsPreviousFrameNumberValid = false;
+		if (!bIgnoreWarningForOneFrame)
+		{
+			UE_LOG(LogRivermaxMedia, Error, TEXT("Rivermax CustomTimeStep wait for sync timed out."));
+		}
+
+		return false;
+	}
+
+	IRivermaxCoreModule& RivermaxModule = FModuleManager::GetModuleChecked<IRivermaxCoreModule>("RivermaxCore");
+	const uint64 CurrentPTPTimeNanosec = RivermaxModule.GetRivermaxManager()->GetTime();
+	const uint64 NewFrameNumber = UE::RivermaxCore::GetFrameNumber(CurrentPTPTimeNanosec, FrameRate);
+	const uint32 ExpectedFramesPerWait = GetExpectedSyncCountDelta();
+
+	if (bEnableOverrunDetection
+		&& bIsPreviousFrameNumberValid
+		&& (NewFrameNumber != (PreviousFrameNumber + ExpectedFramesPerWait)))
+	{
+		UE_LOG(LogRivermaxMedia, Warning,
+			TEXT("The Engine couldn't run fast enough to keep up with the CustomTimeStep Sync. '%d' frame(s) dropped."),
+			NewFrameNumber - PreviousFrameNumber - ExpectedFramesPerWait);
+	}
+
+	{
+		if (bIsPreviousFrameNumberValid)
+		{
+			DeltaFrameNumber = NewFrameNumber - PreviousFrameNumber;
+		}
+
+		PreviousFrameNumber = NewFrameNumber;
+
+		bIsPreviousFrameNumberValid = true;
+	}
+
+	return true;
+}
+
+#if WITH_EDITOR
+void URivermaxCustomTimeStep::PostEditChangeChainProperty(struct FPropertyChangedChainEvent& PropertyChangedEvent)
+{
+	if (PropertyChangedEvent.Property->GetFName() == GET_MEMBER_NAME_CHECKED(URivermaxCustomTimeStep, FrameRate))
+	{
+		// If frame rate is changing live, reset frame number tracking to somewhat reinitialize our state
+		bIsPreviousFrameNumberValid = false;
+	}
+
+	Super::PostEditChangeChainProperty(PropertyChangedEvent);
+}
+#endif // WITH_EDITOR
+
+bool URivermaxCustomTimeStep::WaitForNextFrame()
+{
+	static const double TimeoutSec = 5.0;
+
+	IRivermaxCoreModule& RivermaxModule = FModuleManager::GetModuleChecked<IRivermaxCoreModule>("RivermaxCore");
+	const uint64 CurrentPTPTime = RivermaxModule.GetRivermaxManager()->GetTime();
+	const double CurrentPlatformTime = FPlatformTime::Seconds();
+	const uint64 TargetTimeNanosec = UE::RivermaxCore::GetNextAlignmentPoint(CurrentPTPTime, FrameRate);
+
+	double TimeLeftNanosec = 0.0;
+	if(TargetTimeNanosec >= CurrentPTPTime)
+	{
+		TimeLeftNanosec = TargetTimeNanosec - CurrentPTPTime;
+	}
+
+	const double WaitTimeSec = FMath::Clamp(TimeLeftNanosec / 1E9, 0.0, TimeoutSec);
+	const double StartTime = FPlatformTime::Seconds();
+	double ActualWaitTime = 0.0;
+	{
+		FSimpleScopeSecondsCounter ActualWaitTimeCounter(ActualWaitTime);
+
+		static const float SleepThresholdSec = 5.0f * (1.0f / 1000.0f);
+		static const float SpinningTimeSec = 2.0f * (1.0f / 1000.0f);
+		if (WaitTimeSec > SleepThresholdSec)
+		{
+			FPlatformProcess::SleepNoStats(WaitTimeSec - SpinningTimeSec);
+		}
+
+		while (FPlatformTime::Seconds() < (CurrentPlatformTime + WaitTimeSec))
+		{
+			FPlatformProcess::SleepNoStats(0.f);
+		}
+	}
+
+	TRACE_BOOKMARK(TEXT("PTP genlock at %llu"), TargetTimeNanosec);
+
+	return ActualWaitTime < TimeoutSec;
+}
+
+#undef LOCTEXT_NAMESPACE
diff --git a/Engine/Plugins/VirtualProduction/Rivermax/RivermaxMedia/Source/RivermaxMedia/Public/RivermaxCustomTimeStep.h b/Engine/Plugins/VirtualProduction/Rivermax/RivermaxMedia/Source/RivermaxMedia/Public/RivermaxCustomTimeStep.h
new file mode 100644
--- /dev/null
+++ b/Engine/Plugins/VirtualProduction/Rivermax/RivermaxMedia/Source/RivermaxMedia/Public/RivermaxCustomTimeStep.h
@@ -0,0 +1,71 @@
+// Copyright Epic Games, Inc. All Rights Reserved.
+
+#pragma once
+
+#include "GenlockedCustomTimeStep.h"
+
+#include "RivermaxCustomTimeStep.generated.h"
+
+
+class UEngine;
+
+/**
+ * Genlock using PTP time from a rivermax card
+ */
+UCLASS(Blueprintable, editinlinenew)
+class RIVERMAXMEDIA_API URivermaxCustomTimeStep : public UGenlockedCustomTimeStep
+{
+public:
+	GENERATED_BODY()
+
+	//~ Begin UFixedFrameRateCustomTimeStep interface
+	virtual bool Initialize(UEngine* InEngine) override;
+	virtual void Shutdown(UEngine* InEngine) override;
+	virtual bool UpdateTimeStep(UEngine* InEngine) override;
+	virtual ECustomTimeStepSynchronizationState GetSynchronizationState() const override;
+	virtual FFrameRate GetFixedFrameRate() const override;
+	//~ End UFixedFrameRateCustomTimeStep interface
+
+	//~ UGenlockedCustomTimeStep interface
+	virtual uint32 GetLastSyncCountDelta() const override;
+	virtual bool IsLastSyncDataValid() const override;
+	virtual FFrameRate GetSyncRate() const override;
+	virtual bool WaitForSync() override;
+	//~ End UGenlockedCustomTimeSTep interface
+
+#if WITH_EDITOR
+	virtual void PostEditChangeChainProperty(struct FPropertyChangedChainEvent& PropertyChangedEvent) override;
+#endif // WITH_EDITOR
+
+
+private:
+	bool WaitForNextFrame();
+
+public:
+
+	/** Target frame rate to which to genlock. Uses ST2059 standard to align PTP time to standard genlock */
+	UPROPERTY(EditAnywhere, Category = "Genlock options")
+	FFrameRate FrameRate = FFrameRate(24,1);
+
+	/** When enabled, will warn for skipped frames when engine is too slow */
+	UPROPERTY(EditAnywhere, Category = "Genlock options", meta = (DisplayName = "Display Dropped Frames Warning"), AdvancedDisplay)
+	bool bEnableOverrunDetection = true;
+
+private:
+
+#if WITH_EDITORONLY_DATA
+	/** Engine used to initialize the Provider */
+	UPROPERTY(Transient)
+	TObjectPtr<UEngine> InitializedEngine = nullptr;
+#endif
+
+	/** The current SynchronizationState of this custom timestep */
+	ECustomTimeStepSynchronizationState State = ECustomTimeStepSynchronizationState::Closed;
+
+	bool bIgnoreWarningForOneFrame = true;
+
+	/** Frame number tracking */
+	bool bIsPreviousFrameNumberValid = false;
+	uint64 PreviousFrameNumber = 0;
+	uint64 DeltaFrameNumber = 0;
+};
