Rivermax 2110

- Fix for RTP header and payload not being compliant with standard
- Fix find payload size to only work with equal payload compliant with pgroups
- Adding standard pgroup info
- Fix for resolution not being aligned with pgroup
- Added debugging cvar (stats, scheduling and wakeup offset)
- Fix for PTP time calculation using float

diff --git a/Engine/Plugins/VirtualProduction/Rivermax/RivermaxCore/Source/RivermaxCore/Private/RivermaxFormats.cpp b/Engine/Plugins/VirtualProduction/Rivermax/RivermaxCore/Source/RivermaxCore/Private/RivermaxFormats.cpp
new file mode 100644
--- /dev/null
+++ b/Engine/Plugins/VirtualProduction/Rivermax/RivermaxCore/Source/RivermaxCore/Private/RivermaxFormats.cpp
@@ -0,0 +1,119 @@
+// Copyright Epic Games, Inc. All Rights Reserved.
+
+#include "RivermaxFormats.h"
+
+namespace UE::RivermaxCore
+{
+	const TArray<FVideoFormatInfo> FStandardVideoFormat::AllSamplingTypes =
+	{
+	//	FVideoFormatInfo { Sampling,						PixelDepth, PixelGroupSize, PixelGroupCoverage } }
+		FVideoFormatInfo { ESamplingType::YUV422_8bit,		8,			4,				2 },
+		FVideoFormatInfo { ESamplingType::YUV422_10bit,		10,			5,				2 },
+		FVideoFormatInfo { ESamplingType::YUV422_12bit,		12,			6,				2 },
+		FVideoFormatInfo { ESamplingType::YUV422_16bit,		16,			8,				2 },
+		FVideoFormatInfo { ESamplingType::YUV422_16bitFloat,16,			8,				2 },
+		FVideoFormatInfo { ESamplingType::YUV444_8bit,		8,			3,				1 },
+		FVideoFormatInfo { ESamplingType::YUV444_10bit,		10,			15,				4 },
+		FVideoFormatInfo { ESamplingType::YUV444_12bit,		12,			9,				2 },
+		FVideoFormatInfo { ESamplingType::YUV444_16bit,		16,			6,				1 },
+		FVideoFormatInfo { ESamplingType::YUV444_16bitFloat,16,			6,				1 },
+		FVideoFormatInfo { ESamplingType::RGB_8bit,			8,			3,				1 },
+		FVideoFormatInfo { ESamplingType::RGB_10bit,		10,			15,				4 },
+		FVideoFormatInfo { ESamplingType::RGB_12bit,		12,			9,				2 },
+		FVideoFormatInfo { ESamplingType::RGB_16bit,		16,			6,				1 },
+		FVideoFormatInfo { ESamplingType::RGB_16bitFloat,	16,			6,				1 },
+	};
+
+	FVideoFormatInfo FStandardVideoFormat::GetVideoFormatInfo(ESamplingType SamplingType)
+	{
+		if (const FVideoFormatInfo* FoundInfo = AllSamplingTypes.FindByPredicate([SamplingType](const FVideoFormatInfo& Other) { return Other.Sampling == SamplingType; }))
+		{
+			return *FoundInfo;
+		}
+
+		return FVideoFormatInfo();
+	}
+
+	bool FStandardVideoFormat::IsRGB(ESamplingType SamplingType)
+	{
+		switch (SamplingType)
+		{
+		case ESamplingType::RGB_8bit:
+		case ESamplingType::RGB_10bit:
+		case ESamplingType::RGB_12bit:
+		case ESamplingType::RGB_16bit:
+		case ESamplingType::RGB_16bitFloat:
+		{
+			return true;
+		}
+		default:
+		{
+			return false;
+		}
+		};
+	}
+
+	bool FStandardVideoFormat::IsYUV(ESamplingType SamplingType)
+	{
+		switch (SamplingType)
+		{
+		case ESamplingType::YUV422_8bit:
+		case ESamplingType::YUV422_10bit:
+		case ESamplingType::YUV422_12bit:
+		case ESamplingType::YUV422_16bit:
+		case ESamplingType::YUV422_16bitFloat:
+		case ESamplingType::YUV444_8bit:
+		case ESamplingType::YUV444_10bit:
+		case ESamplingType::YUV444_12bit:
+		case ESamplingType::YUV444_16bit:
+		case ESamplingType::YUV444_16bitFloat:
+		{
+			return true;
+		}
+		default:
+		{
+			return false;
+		}
+		};
+	}
+
+	bool FStandardVideoFormat::IsYUV422(ESamplingType SamplingType)
+	{
+		switch (SamplingType)
+		{
+		case ESamplingType::YUV422_8bit:
+		case ESamplingType::YUV422_10bit:
+		case ESamplingType::YUV422_12bit:
+		case ESamplingType::YUV422_16bit:
+		case ESamplingType::YUV422_16bitFloat:
+		{
+			return true;
+		}
+		default:
+		{
+			return false;
+		}
+		};
+	}
+
+	bool FStandardVideoFormat::IsYUV444(ESamplingType SamplingType)
+	{
+		switch (SamplingType)
+		{
+		case ESamplingType::YUV444_8bit:
+		case ESamplingType::YUV444_10bit:
+		case ESamplingType::YUV444_12bit:
+		case ESamplingType::YUV444_16bit:
+		case ESamplingType::YUV444_16bitFloat:
+		{
+			return true;
+		}
+		default:
+		{
+			return false;
+		}
+		};
+	}
+}
+
+	
diff --git a/Engine/Plugins/VirtualProduction/Rivermax/RivermaxCore/Source/RivermaxCore/Private/RivermaxInputStream.cpp b/Engine/Plugins/VirtualProduction/Rivermax/RivermaxCore/Source/RivermaxCore/Private/RivermaxInputStream.cpp
--- a/Engine/Plugins/VirtualProduction/Rivermax/RivermaxCore/Source/RivermaxCore/Private/RivermaxInputStream.cpp
+++ b/Engine/Plugins/VirtualProduction/Rivermax/RivermaxCore/Source/RivermaxCore/Private/RivermaxInputStream.cpp
@@ -107,6 +107,7 @@ namespace UE::RivermaxCore::Private
 
 		Options = InOptions;
 		Listener = &InListener;
+		FormatInfo = FStandardVideoFormat::GetVideoFormatInfo(Options.PixelFormat);
 
 		InitTaskFuture = Async(EAsyncExecution::TaskGraph, [this]()
 		{
@@ -402,7 +403,7 @@ namespace UE::RivermaxCore::Private
 									FRivermaxInputVideoFrameDescriptor Descriptor;
 									Descriptor.Width = Options.Resolution.X;
 									Descriptor.Height = Options.Resolution.Y;
-									Descriptor.Stride = Options.Stride;
+									Descriptor.Stride = Options.AlignedResolution.X / FormatInfo.PixelGroupCoverage * FormatInfo.PixelGroupSize;
 									FRivermaxInputVideoFrameReception NewFrame;
 									NewFrame.VideoBuffer = StreamData.CurrentFrame;
 									Listener->OnVideoFrameReceived(Descriptor, NewFrame);
@@ -461,7 +462,8 @@ namespace UE::RivermaxCore::Private
 
 		FRivermaxInputVideoFrameDescriptor Descriptor;
 		FRivermaxInputVideoFrameRequest Request;
-		Descriptor.VideoBufferSize = Options.Resolution.Y * Options.Stride;
+		const int32 Stride = Options.AlignedResolution.X / FormatInfo.PixelGroupCoverage * FormatInfo.PixelGroupSize;
+		Descriptor.VideoBufferSize = Options.Resolution.Y * Stride;
 		Listener->OnVideoFrameRequested(Descriptor, Request);
 		StreamData.CurrentFrame = Request.VideoBuffer;
 		StreamData.WritingOffset = 0;
diff --git a/Engine/Plugins/VirtualProduction/Rivermax/RivermaxCore/Source/RivermaxCore/Private/RivermaxInputStream.h b/Engine/Plugins/VirtualProduction/Rivermax/RivermaxCore/Source/RivermaxCore/Private/RivermaxInputStream.h
--- a/Engine/Plugins/VirtualProduction/Rivermax/RivermaxCore/Source/RivermaxCore/Private/RivermaxInputStream.h
+++ b/Engine/Plugins/VirtualProduction/Rivermax/RivermaxCore/Source/RivermaxCore/Private/RivermaxInputStream.h
@@ -110,7 +110,7 @@ namespace UE::RivermaxCore::Private
 		std::atomic<bool> bIsShuttingDown;
 		TFuture<void> InitTaskFuture;
 		double LastLoggingTimestamp = 0.0;
-
+		FVideoFormatInfo FormatInfo;
 	};
 }
 
diff --git a/Engine/Plugins/VirtualProduction/Rivermax/RivermaxCore/Source/RivermaxCore/Private/RivermaxManager.cpp b/Engine/Plugins/VirtualProduction/Rivermax/RivermaxCore/Source/RivermaxCore/Private/RivermaxManager.cpp
--- a/Engine/Plugins/VirtualProduction/Rivermax/RivermaxCore/Source/RivermaxCore/Private/RivermaxManager.cpp
+++ b/Engine/Plugins/VirtualProduction/Rivermax/RivermaxCore/Source/RivermaxCore/Private/RivermaxManager.cpp
@@ -132,7 +132,7 @@ namespace UE::RivermaxCore::Private
 		const URivermaxSettings* Settings = GetDefault<URivermaxSettings>();
 
 		bool bConfiguredSuccessfully = false;
-		if (TimeSource == ERivermaxTimeSource::PTP)
+		if (Settings->TimeSource == ERivermaxTimeSource::PTP)
 		{
 			Config.clock_configurations.clock_type = rmax_clock_types::RIVERMAX_PTP_CLOCK;
 
diff --git a/Engine/Plugins/VirtualProduction/Rivermax/RivermaxCore/Source/RivermaxCore/Private/RivermaxOutputStream.cpp b/Engine/Plugins/VirtualProduction/Rivermax/RivermaxCore/Source/RivermaxCore/Private/RivermaxOutputStream.cpp
--- a/Engine/Plugins/VirtualProduction/Rivermax/RivermaxCore/Source/RivermaxCore/Private/RivermaxOutputStream.cpp
+++ b/Engine/Plugins/VirtualProduction/Rivermax/RivermaxCore/Source/RivermaxCore/Private/RivermaxOutputStream.cpp
@@ -15,35 +15,69 @@
 namespace UE::RivermaxCore::Private
 {
 	static TAutoConsoleVariable<int32> CVarRivermaxWakeupOffset(
-		TEXT("Rivermax.AlignmentOffset"), 0,
-		TEXT("Offset from alignment point to wake up in nanoseconds"),
+		TEXT("Rivermax.WakeupOffset"), 0,
+		TEXT("Wakeup is done on alignment point. This offset will be substracted from it to wake up earlier. Units: nanoseconds"),
 		ECVF_Default);
 
-	uint32 FindPayloadSize(const FRivermaxStreamOptions& InOptions, uint32 InBytesPerLine)
+	static TAutoConsoleVariable<int32> CVarRivermaxScheduleOffset(
+		TEXT("Rivermax.ScheduleOffset"), 0,
+		TEXT("Scheduling is done at alignment point plus TRO. This offset will be added to it to delay or schedule earlier. Units: nanoseconds"),
+		ECVF_Default);
+
+	static TAutoConsoleVariable<int32> CVarRivermaxOutputShowStats(
+		TEXT("Rivermax.ShowOutputStats"), 0,
+		TEXT("Enable stats logging at fixed interval"),
+		ECVF_Default);
+
+	static TAutoConsoleVariable<float> CVarRivermaxOutputShowStatsInterval(
+		TEXT("Rivermax.ShowStatsInterval"), 1.0,
+		TEXT("Interval at which to show stats in seconds"),
+		ECVF_Default);
+
+
+	bool FindPayloadSize(const FRivermaxStreamOptions& InOptions, uint32 InBytesPerLine, const FVideoFormatInfo& FormatInfo, uint16& OutPayloadSize)
 	{
 		using namespace UE::RivermaxCore::Private::Utils;
 
-		// If stride per line is small then our minimal payload size, send one line per packet.
-		if (InBytesPerLine <= MaxPayloadSize)
+		// For now, we only find a payload size that can be equal across one line
+		// Support for last payload of a line being smaller is there but is causing issue
+		// We fail to output if we receive a resolution for which we can't find an equal payload size
+		int32 TestPoint = InBytesPerLine / MaxPayloadSize;
+		if (TestPoint == 0)
 		{
-			return InBytesPerLine;
+			if (InBytesPerLine > MinPayloadSize)
+			{
+				if (InBytesPerLine % FormatInfo.PixelGroupSize == 0)
+				{
+					OutPayloadSize = InBytesPerLine;
+					return true;
+				}
+			}
+			return false;
 		}
 
-		uint32 PacketDivider = 2; //Start with 2 packets per line as 1 is too big already
-		static constexpr uint32 MaxPacketCountPerLine = 50; 
-		while (PacketDivider <= MaxPacketCountPerLine)
+		bool bContinue = true;
+		while (bContinue)
 		{
-			const uint32 TestPayloadSize = InBytesPerLine / PacketDivider;
-			const bool bIsFlushInLine = InBytesPerLine % PacketDivider == 0;
-			if (bIsFlushInLine && TestPayloadSize < MaxPayloadSize)
+			const int32 TestSize = InBytesPerLine / TestPoint;
+			if (TestSize < MinPayloadSize)
 			{
-				return TestPayloadSize;
+				return false;
 			}
 
-			++PacketDivider;
+			if (TestSize <= MaxPayloadSize)
+			{
+				if ((TestSize % FormatInfo.PixelGroupSize) == 0 && (InBytesPerLine % TestPoint) == 0)
+				{
+					OutPayloadSize = TestSize;
+					return true;
+				}
+			}
+
+			++TestPoint;
 		}
 
-		return 0;
+		return false;
 	}
 
 	uint32 FindChunksPerLine(const FRivermaxStreamOptions& InOptions)
@@ -86,7 +120,14 @@ namespace UE::RivermaxCore::Private
 
 		Options = InOptions;
 		Listener = &InListener;
-		AlignmentPointSchedulingLeadTimeNanosec = CVarRivermaxWakeupOffset.GetValueOnGameThread();
+		FormatInfo = FStandardVideoFormat::GetVideoFormatInfo(Options.PixelFormat);
+
+		// Verify resolution for sampling type
+		if (Options.AlignedResolution.X % FormatInfo.PixelGroupCoverage != 0)
+		{
+			UE_LOG(LogRivermax, Warning, TEXT("Can't create Rivermax Output Stream. Aligned horizontal resolution of %d doesn't align with pixel group coverage of %d."), Options.AlignedResolution.X, FormatInfo.PixelGroupCoverage);
+			return false;
+		}
 
 		Async(EAsyncExecution::TaskGraph, [this]()
 		{
@@ -100,43 +141,44 @@ namespace UE::RivermaxCore::Private
 			InitializeBuffers();
 
 			// Initialize internal memory and rivermax configuration 
-			InitializeMemory();
-
-			// Initialize rivermax output stream with desired config
-			const uint32 NumberPacketsPerFrame = StreamMemory.PacketsInFrameField;
-			const uint32 MediaBlockIndex = 0;
-			rmax_stream_id NewId;
-			rmax_qos_attr QOSAttributes = { 0, 0 }; //todo
-
-			rmax_buffer_attr BufferAttributes;
-			BufferAttributes.chunk_size_in_strides = StreamMemory.ChunkSizeInStrides;
-			BufferAttributes.data_stride_size = StreamMemory.PayloadSize; //Stride between chunks. 
-			BufferAttributes.app_hdr_stride_size = StreamMemory.HeaderStrideSize; 
-			BufferAttributes.mem_block_array = StreamMemory.MemoryBlocks.GetData();
-			BufferAttributes.mem_block_array_len = StreamMemory.MemoryBlocks.Num();
-
-			rmax_status_t Status = rmax_out_create_stream(SDPDescription.GetData(), &BufferAttributes, &QOSAttributes, NumberPacketsPerFrame, MediaBlockIndex, &NewId);
-			if (Status == RMAX_OK)
-			{
-				struct sockaddr_in SourceAddress;
-				struct sockaddr_in DestinationAddress;
-				Status = rmax_out_query_address(NewId, MediaBlockIndex, &SourceAddress, &DestinationAddress);
+			if(InitializeMemory())
+			{
+				// Initialize rivermax output stream with desired config
+				const uint32 NumberPacketsPerFrame = StreamMemory.PacketsPerFrame;
+				const uint32 MediaBlockIndex = 0;
+				rmax_stream_id NewId;
+				rmax_qos_attr QOSAttributes = { 0, 0 }; //todo
+
+				rmax_buffer_attr BufferAttributes;
+				BufferAttributes.chunk_size_in_strides = StreamMemory.ChunkSizeInStrides;
+				BufferAttributes.data_stride_size = StreamMemory.PayloadSize; //Stride between chunks. 
+				BufferAttributes.app_hdr_stride_size = StreamMemory.HeaderStrideSize;
+				BufferAttributes.mem_block_array = StreamMemory.MemoryBlocks.GetData();
+				BufferAttributes.mem_block_array_len = StreamMemory.MemoryBlocks.Num();
+
+				rmax_status_t Status = rmax_out_create_stream(SDPDescription.GetData(), &BufferAttributes, &QOSAttributes, NumberPacketsPerFrame, MediaBlockIndex, &NewId);
 				if (Status == RMAX_OK)
 				{
-					StreamId = NewId;
+					struct sockaddr_in SourceAddress;
+					struct sockaddr_in DestinationAddress;
+					Status = rmax_out_query_address(NewId, MediaBlockIndex, &SourceAddress, &DestinationAddress);
+					if (Status == RMAX_OK)
+					{
+						StreamId = NewId;
 
-					// Todo add event manager to avoid spinning
+						StreamData.FrameFieldTimeIntervalNs = 1E9 / Options.FrameRate.AsDecimal();
+						InitializeStreamTimingSettings();
 
-					StreamData.FrameFieldTimeIntervalNs = 1E9 / Options.FrameRate.AsDecimal();
-					InitializeStreamTimingSettings();
+						UE_LOG(LogRivermax, Display, TEXT("Output started to send %dx%d using %d payloads of size %d"), Options.AlignedResolution.X, Options.AlignedResolution.Y, StreamMemory.PacketsInLine, StreamMemory.PayloadSize);
 
-					bIsActive = true;
-					RivermaxThread.Reset(FRunnableThread::Create(this, TEXT("Rmax OutputStream Thread"), 128 * 1024, TPri_AboveNormal, FPlatformAffinity::GetPoolThreadMask()));
+						bIsActive = true;
+						RivermaxThread.Reset(FRunnableThread::Create(this, TEXT("Rmax OutputStream Thread"), 128 * 1024, TPri_AboveNormal, FPlatformAffinity::GetPoolThreadMask()));
+					}
+				}
+				else
+				{
+					UE_LOG(LogRivermax, Warning, TEXT("Failed to create Rivermax output stream. Status: %d"), Status);
 				}
-			}
-			else
-			{
-				UE_LOG(LogRivermax, Warning, TEXT("Failed to create Rivermax output stream. Status: %d"), Status);
 			}
 
 			Listener->OnInitializationCompleted(bIsActive);
@@ -168,7 +210,9 @@ namespace UE::RivermaxCore::Private
 	{
 		if (TSharedPtr<FRivermaxOutputFrame> AvailableFrame = GetNextAvailableFrame(NewFrame.FrameIdentifier))
 		{
-			FMemory::Memcpy(AvailableFrame->VideoBuffer, NewFrame.VideoBuffer, NewFrame.Height * Options.Stride);
+			const int32 Stride = Options.AlignedResolution.X / FormatInfo.PixelGroupCoverage * FormatInfo.PixelGroupSize;
+			FMemory::Memcpy(AvailableFrame->VideoBuffer, NewFrame.VideoBuffer, NewFrame.Height * Stride);
+
 			AvailableFrame->bIsVideoBufferReady = true;
 
 			//If Frame ready to be sent
@@ -250,42 +294,53 @@ namespace UE::RivermaxCore::Private
 			Stats.MemoryBlockSentCounter++;
 			Stats.TotalStrides += StreamMemory.ChunkSizeInStrides;
 			StreamData.bHasFrameFirstChunkBeenFetched = false;
-			UE_LOG(LogRivermax, Verbose, TEXT("Stats: BlockSent: %d. CommitImmediate: %d. CommitRetries: %d. "), Stats.MemoryBlockSentCounter, Stats.CommitImmediate, Stats.CommitRetries);
-
 		}
 	}
 
 	void FRivermaxOutputStream::InitializeBuffers()
 	{
+		const int32 Stride = Options.AlignedResolution.X / FormatInfo.PixelGroupCoverage * FormatInfo.PixelGroupSize;
+		const int32 FrameAllocSize = Options.AlignedResolution.Y * Stride;
 		AvailableFrames.Reserve(Options.NumberOfBuffers);
 		for (int32 Index = 0; Index < Options.NumberOfBuffers; ++Index)
 		{
 			TSharedPtr<FRivermaxOutputFrame> Frame = MakeShared<FRivermaxOutputFrame>(Index);
-			Frame->VideoBuffer = static_cast<uint8*>(FMemory::Malloc(Options.Resolution.Y * Options.Stride));
+
+			constexpr uint32 CacheLineSize = PLATFORM_CACHE_LINE_SIZE;
+			Frame->VideoBuffer = static_cast<uint8*>(FMemory::Malloc(FrameAllocSize, CacheLineSize));
 			AvailableFrames.Add(MoveTemp(Frame));
 		}
 	}
 
-	void FRivermaxOutputStream::InitializeMemory()
+	bool FRivermaxOutputStream::InitializeMemory()
 	{
 		using namespace UE::RivermaxCore::Private::Utils;
 
 		//2110 stream type
 		//We need to use the fullframe allocated size to compute the payload size.
 
-		const uint32 EffectiveBytesPerLine = Options.Stride;
-		const uint32 BytesPerLine = Options.Stride;
-		StreamMemory.PayloadSize = FindPayloadSize(Options, BytesPerLine); 
-		StreamMemory.DataStrideSize = BytesPerLine;//todo align cache line?
+		const int32 BytesPerLine = Options.AlignedResolution.X / FormatInfo.PixelGroupCoverage * FormatInfo.PixelGroupSize;
+		const uint32 EffectiveBytesPerLine = BytesPerLine;
+
+		const bool bFoundPayload = FindPayloadSize(Options, BytesPerLine, FormatInfo, StreamMemory.PayloadSize);
+		if (bFoundPayload == false)
+		{
+			UE_LOG(LogRivermax, Warning, TEXT("Could not find payload size for desired resolution %dx%d for desired pixel format"), Options.AlignedResolution.X, Options.AlignedResolution.Y);
+			return false;
+		}
+
+		StreamMemory.PixelGroupPerPacket = StreamMemory.PayloadSize / FormatInfo.PixelGroupSize;
+		StreamMemory.PixelsPerPacket = StreamMemory.PixelGroupPerPacket * FormatInfo.PixelGroupCoverage;
+		StreamMemory.PacketsInLine = FMath::CeilToInt32((float)Options.AlignedResolution.X / StreamMemory.PixelsPerPacket);
+
 		StreamMemory.LinesInChunk = FindChunksPerLine(Options);
 
-		StreamMemory.PacketsInLine =  FMath::RoundToInt32(BytesPerLine / (float)StreamMemory.PayloadSize);
 		StreamMemory.ChunkSizeInStrides = StreamMemory.LinesInChunk * StreamMemory.PacketsInLine;
 
 		StreamMemory.FramesFieldPerMemoryBlock = 1;
-		StreamMemory.PacketsInFrameField =  StreamMemory.PacketsInLine * Options.Resolution.Y;
-		StreamMemory.PacketsPerMemoryBlock = StreamMemory.PacketsInFrameField * StreamMemory.FramesFieldPerMemoryBlock;
-		StreamMemory.ChunksPerFrameField = StreamMemory.PacketsInFrameField / StreamMemory.ChunkSizeInStrides;
+		StreamMemory.PacketsPerFrame =  StreamMemory.PacketsInLine * Options.Resolution.Y;
+		StreamMemory.PacketsPerMemoryBlock = StreamMemory.PacketsPerFrame * StreamMemory.FramesFieldPerMemoryBlock;
+		StreamMemory.ChunksPerFrameField = StreamMemory.PacketsPerFrame / StreamMemory.ChunkSizeInStrides;
 		StreamMemory.ChunksPerMemoryBlock = StreamMemory.FramesFieldPerMemoryBlock * StreamMemory.ChunksPerFrameField;
 		StreamMemory.MemoryBlockCount = Options.NumberOfBuffers;
 		StreamMemory.StridesPerMemoryBlock = StreamMemory.ChunkSizeInStrides * StreamMemory.ChunksPerMemoryBlock;
@@ -294,13 +349,15 @@ namespace UE::RivermaxCore::Private
 		StreamMemory.RTPHeaders.SetNumZeroed(StreamMemory.PacketsPerMemoryBlock);
 		StreamMemory.PayloadSizes.SetNumUninitialized(StreamMemory.PacketsPerMemoryBlock);
 		StreamMemory.HeaderSizes.SetNumUninitialized(StreamMemory.PacketsPerMemoryBlock);
-		StreamMemory.HeaderStrideSize = BytesPerHeader;
+		StreamMemory.HeaderStrideSize = RTPHeaderSize;
 		for (int32 PayloadSizeIndex = 0; PayloadSizeIndex < StreamMemory.PayloadSizes.Num(); ++PayloadSizeIndex)
 		{
 			//Go through each chunk to have effective payload size to be sent (last one of each line could be smaller)
 			if ((PayloadSizeIndex + 1) % StreamMemory.PacketsInLine == 0)
 			{
-				StreamMemory.PayloadSizes[PayloadSizeIndex] = EffectiveBytesPerLine - ((StreamMemory.PacketsInLine - 1) * StreamMemory.PayloadSize);
+				const uint32 LeftOver = EffectiveBytesPerLine - ((StreamMemory.PacketsInLine - 1) * StreamMemory.PayloadSize);
+				StreamMemory.PayloadSizes[PayloadSizeIndex] = LeftOver;
+				check(LeftOver % FormatInfo.PixelGroupSize == 0);
 			}
 			else
 			{
@@ -318,6 +375,8 @@ namespace UE::RivermaxCore::Private
 			Block.data_ptr = AvailableFrames[BlockIndex]->VideoBuffer;
 			Block.app_hdr_ptr = &StreamMemory.RTPHeaders[BlockIndex];
 		}
+
+		return true;
 	}
 
 	void FRivermaxOutputStream::InitializeNextFrame(const TSharedPtr<FRivermaxOutputFrame>& NextFrame)
@@ -415,6 +474,7 @@ namespace UE::RivermaxCore::Private
 		if (StreamType == ERivermaxStreamType::VIDEO_2110_20_STREAM)
 		{
 			//SRD means Sample Row Data
+			
 			// build SRD header - 8-14 bytes
 		   /* 0                   1                   2                   3
 			0 1 2 3 4 5 6 7 8 9 0 1 2 3 4 5 6 7 8 9 0 1 2 3 4 5 6 7 8 9 0 1
@@ -426,7 +486,7 @@ namespace UE::RivermaxCore::Private
 			OutHeader.RawHeader[12] = (StreamData.SequenceNumber >> 24) & 0xff;  // high 16 bit of seq Extended Sequence Number
 			OutHeader.RawHeader[13] = (StreamData.SequenceNumber >> 16) & 0xff;  // high 16 bit of seq Extended Sequence Number
 
-			const uint16 CurrentPayloadSize = StreamMemory.PayloadSizes[StreamData.SequenceNumber % StreamMemory.PacketsInFrameField];
+			const uint16 CurrentPayloadSize = StreamMemory.PayloadSizes[StreamData.SequenceNumber % StreamMemory.PacketsPerFrame];
 			*(uint16*)&OutHeader.RawHeader[14] = ByteSwap(CurrentPayloadSize);  // SRD Length
 		
 			uint16 number_of_rows = Options.Resolution.Y; //todo divide by 2 if interlaced
@@ -435,14 +495,19 @@ namespace UE::RivermaxCore::Private
 			OutHeader.RawHeader[16] |= (0 << 7); //todo when fields are sent for interlace
 
 			// we never have continuation
+
+			// Write out current offset in pixels
 			*(uint16*)&OutHeader.RawHeader[18] = ByteSwap(CurrentFrame->SRDOffset);  // SRD Offset
-			uint16 group_size = (uint16)((StreamMemory.PayloadSize /*- 20*/) / 2.5); // todo: need to support pgroup definition per sampling format to fully be compliant
-			CurrentFrame->SRDOffset = (CurrentFrame->SRDOffset + group_size) % (group_size * StreamMemory.PacketsInLine);
+
+			// Update SRD offset in pixel for next round
+			uint16 PixelsPerPacket = (uint16)((StreamMemory.PayloadSize * FormatInfo.PixelGroupCoverage) / FormatInfo.PixelGroupSize);
+			CurrentFrame->SRDOffset = (CurrentFrame->SRDOffset + PixelsPerPacket) % (PixelsPerPacket * StreamMemory.PacketsInLine);
 		}
 
-		if (++CurrentFrame->PacketCounter == StreamMemory.PacketsInFrameField)
+		if (++CurrentFrame->PacketCounter == StreamMemory.PacketsPerFrame)
 		{
 			OutHeader.RawHeader[1] |= 0x80; // last packet in frame (Marker)
+
 			// ST2210-20: the timestamp SHOULD be the same for each packet of the frame/field.
 			const double ticks = (MediaClockSampleRate / (Options.FrameRate.AsDecimal()));
 			//if (set.video_type != VIDEO_TYPE::PROGRESSIVE) {
@@ -528,18 +593,17 @@ namespace UE::RivermaxCore::Private
 		// Get next alignment point based on the frame number we are aligning with
 		const uint64 NextAlignmentNano = UE::RivermaxCore::GetAlignmentPointFromFrameNumber(NextFrameNumber, Options.FrameRate);
 
-		// Add Tro offset to next alignment point
+		// Add Tro offset to next alignment point and configurable offset
 		StreamData.NextAlignmentPointNanosec = NextAlignmentNano;
-		StreamData.NextScheduleTimeNanosec = NextAlignmentNano + TransmitOffsetNanosec;
+		StreamData.NextScheduleTimeNanosec = NextAlignmentNano + TransmitOffsetNanosec + CVarRivermaxScheduleOffset.GetValueOnAnyThread();
 		StreamData.NextAlignmentPointFrameNumber = NextFrameNumber;
 
 		// Offset wakeup if desired to give more time for scheduling. 
-		const uint64 WakeupTime = NextAlignmentNano - AlignmentPointSchedulingLeadTimeNanosec;
+		const uint64 WakeupTime = NextAlignmentNano - CVarRivermaxWakeupOffset.GetValueOnAnyThread();
 
 		// Used to know if we schedule blindly in the future or based on previous data
 		StreamData.bHasValidNextFrameNumber = bHasValidTimings;
 		
-
 		uint64 WaitTimeNanosec = WakeupTime - CurrentTimeNanosec;
 
 		// Wakeup can be smaller than current time with controllable offset
@@ -628,7 +692,7 @@ namespace UE::RivermaxCore::Private
 	{
 		uint8* HeaderRawPtr = reinterpret_cast<uint8*>(CurrentFrame->HeaderPtr);
 		check(HeaderRawPtr);
-		for (uint32 StrideIndex = 0; StrideIndex < StreamMemory.ChunkSizeInStrides && CurrentFrame->PacketCounter < StreamMemory.PacketsInFrameField; ++StrideIndex)
+		for (uint32 StrideIndex = 0; StrideIndex < StreamMemory.ChunkSizeInStrides && CurrentFrame->PacketCounter < StreamMemory.PacketsPerFrame; ++StrideIndex)
 		{
 			uint8* NextHeaderRawPtr = HeaderRawPtr + (StrideIndex * StreamMemory.HeaderStrideSize);
 			BuildRTPHeader(*reinterpret_cast<FRTPHeader*>(NextHeaderRawPtr));
@@ -646,8 +710,6 @@ namespace UE::RivermaxCore::Private
 		rmax_status_t Status;
 		do
 		{
-			// todo configure timeout
-			
 			//Only first chunk gets scheduled with a timestamp. Following chunks are queued after it using 0
 			uint64 ScheduleTime = CurrentFrame->ChunkNumber == 0 ? StreamData.NextScheduleTimeNanosec : 0;
 			const rmax_commit_flags_t CommitFlags{};
@@ -655,6 +717,7 @@ namespace UE::RivermaxCore::Private
 			{
 				IRivermaxCoreModule& RivermaxModule = FModuleManager::LoadModuleChecked<IRivermaxCoreModule>(TEXT("RivermaxCore"));
 				const uint64 CurrentTimeNanosec = RivermaxModule.GetRivermaxManager()->GetTime();
+				TRACE_BOOKMARK(TEXT("Sched A: %llu, C: %llu"), StreamData.NextAlignmentPointNanosec, CurrentTimeNanosec);
 				if (ScheduleTime <= CurrentTimeNanosec)
 				{
 					ScheduleTime = 0;
@@ -671,6 +734,7 @@ namespace UE::RivermaxCore::Private
 			else if (Status == RMAX_ERR_HW_SEND_QUEUE_FULL)
 			{
 				Stats.CommitRetries++;
+				TRACE_CPUPROFILER_EVENT_SCOPE(CommitNextChunks::QUEUEFULL);
 			}
 			else if(Status == RMAX_ERR_HW_COMPLETION_ISSUE)
 			{
@@ -695,12 +759,13 @@ namespace UE::RivermaxCore::Private
 
 	uint32 FRivermaxOutputStream::Run()
 	{
+		ReadyToSendEvent->Wait();
 		while (bIsActive)
 		{
+			ShowStats();
 			Process_AnyThread();
 		}
 
-
 		DestroyStream();
 
 		return 0;
@@ -763,15 +828,26 @@ namespace UE::RivermaxCore::Private
 				InitializeNextFrame(CurrentFrame);		
 				
 				// Since we want to resend last frame, we need to fast forward chunk pointer to re-point to the one we just sent
-				rmax_status_t Status = rmax_out_skip_chunks(StreamId, StreamMemory.ChunksPerFrameField * (Options.NumberOfBuffers - 1));
-				if (Status != RMAX_OK)
+				rmax_status_t Status;
+				do 
 				{
-					// If we don't go in error, we end up with misaligned chunk vs frame pointer to write. 
-					ensure(false);
-					UE_LOG(LogRivermax, Error, TEXT("Invalid error happened while trying to skip chunks. Status: %d."), Status);
-					Listener->OnStreamError();
-					bIsActive = false;
-				}
+					Status = rmax_out_skip_chunks(StreamId, StreamMemory.ChunksPerFrameField * (Options.NumberOfBuffers - 1));
+					if (Status != RMAX_OK)
+					{
+						if (Status == RMAX_ERR_NO_FREE_CHUNK)
+						{
+							// Wait until there are enough free chunk to be skipped
+							UE_LOG(LogRivermax, Warning, TEXT("No chunks ready to skip. Waiting"));
+						}
+						else
+						{ 
+							ensure(false);
+							UE_LOG(LogRivermax, Error, TEXT("Invalid error happened while trying to skip chunks. Status: %d."), Status);
+							Listener->OnStreamError();
+							bIsActive = false;
+						}
+					}
+				} while (Status != RMAX_OK && bIsActive);
 			}
 		}
 
@@ -784,7 +860,7 @@ namespace UE::RivermaxCore::Private
 
 		double FrameIntervalNs = StreamData.FrameFieldTimeIntervalNs;
 		const bool bIsProgressive = true;//todo MediaConfiguration.IsProgressive() 
-		uint32 PacketsInFrameField = StreamMemory.PacketsInFrameField;
+		uint32 PacketsInFrameField = StreamMemory.PacketsPerFrame;
 		if (bIsProgressive == false)
 		{
 			FrameIntervalNs *= 2;
@@ -832,7 +908,7 @@ namespace UE::RivermaxCore::Private
 
 			// Need to reinvestigate the implication of this and possibly add cvar to control it runtime
 			const double TRSNano = (FrameIntervalNs * RActive) / PacketsInFrameField;
-			TransmitOffsetNanosec = (uint64)((TRODefaultMultiplier * FrameIntervalNs) - (VideoTROModification * TRSNano));
+			TransmitOffsetNanosec = (uint64)((TRODefaultMultiplier * FrameIntervalNs));
 		}
 	}
 
@@ -852,5 +928,17 @@ namespace UE::RivermaxCore::Private
 		return MediaTime + MediaSubTime;
 	}
 
+	void FRivermaxOutputStream::ShowStats()
+	{
+		if (CVarRivermaxOutputShowStats.GetValueOnAnyThread() != 0)
+		{
+			const double CurrentTime = FPlatformTime::Seconds();
+			if (CurrentTime - LastStatsShownTimestamp > CVarRivermaxOutputShowStatsInterval.GetValueOnAnyThread())
+			{
+				LastStatsShownTimestamp = CurrentTime;
+				UE_LOG(LogRivermax, Log, TEXT("Stats: FrameSent: %d. CommitImmediate: %d. CommitRetries: %d. ChunkRetries: %d"), Stats.MemoryBlockSentCounter, Stats.CommitImmediate, Stats.CommitRetries, Stats.ChunkRetries);
+			}
+		}
+	}
 }
 
diff --git a/Engine/Plugins/VirtualProduction/Rivermax/RivermaxCore/Source/RivermaxCore/Private/RivermaxOutputStream.h b/Engine/Plugins/VirtualProduction/Rivermax/RivermaxCore/Source/RivermaxCore/Private/RivermaxOutputStream.h
--- a/Engine/Plugins/VirtualProduction/Rivermax/RivermaxCore/Source/RivermaxCore/Private/RivermaxOutputStream.h
+++ b/Engine/Plugins/VirtualProduction/Rivermax/RivermaxCore/Source/RivermaxCore/Private/RivermaxOutputStream.h
@@ -26,7 +26,10 @@ namespace UE::RivermaxCore::Private
 	struct FRivermaxOutputStreamMemory
 	{
 		uint16 PayloadSize = 0;
-		uint32 DataStrideSize = 1280;
+		uint32 PixelGroupPerPacket = 0;
+		uint32 PixelsPerPacket = 0;
+		uint32 PixelsPerFrame = 0;
+
 		uint32 HeaderStrideSize = 20;
 		uint32 LinesInChunk = 4;
 
@@ -34,7 +37,7 @@ namespace UE::RivermaxCore::Private
 		uint32 ChunkSizeInStrides = 0;
 
 		uint32 FramesFieldPerMemoryBlock = 0;
-		uint32 PacketsInFrameField = 0;
+		uint32 PacketsPerFrame = 0;
 		uint32 PacketsPerMemoryBlock = 0;
 		uint32 ChunksPerFrameField = 0;
 		uint32 ChunksPerMemoryBlock = 0;
@@ -108,7 +111,7 @@ namespace UE::RivermaxCore::Private
 
 	private:
 		void InitializeBuffers();
-		void InitializeMemory();
+		bool InitializeMemory();
 		void InitializeNextFrame(const TSharedPtr<FRivermaxOutputFrame>& NextFrame);
 		TSharedPtr<FRivermaxOutputFrame> GetNextFrameToSend();
 		TSharedPtr<FRivermaxOutputFrame> GetNextAvailableFrame(uint32 InFrameIdentifier);
@@ -120,6 +123,7 @@ namespace UE::RivermaxCore::Private
 		void CommitNextChunks();
 		void PrepareNextFrame();
 		void InitializeStreamTimingSettings();
+		void ShowStats();
 		uint32 GetTimestampFromTime(uint64 InTimeNanosec, double InMediaClockRate) const;
 
 	private:
@@ -149,8 +153,11 @@ namespace UE::RivermaxCore::Private
 		/** TRoffset time calculated based on ST2110 - 21 Gapped(for now) method. This is added to next alignment point */
 		uint64 TransmitOffsetNanosec = 0;
 
-		/** Offset substracted to next alignment point (wake up) to move the start of scheduling */
-		uint64 AlignmentPointSchedulingLeadTimeNanosec = 0;
+		/** Format info for the active stream */
+		FVideoFormatInfo FormatInfo;
+
+		/** Timestamp at which we logged stats */
+		double LastStatsShownTimestamp = 0.0;
 	};
 }
 
diff --git a/Engine/Plugins/VirtualProduction/Rivermax/RivermaxCore/Source/RivermaxCore/Private/RivermaxPTPUtils.cpp b/Engine/Plugins/VirtualProduction/Rivermax/RivermaxCore/Source/RivermaxCore/Private/RivermaxPTPUtils.cpp
--- a/Engine/Plugins/VirtualProduction/Rivermax/RivermaxCore/Source/RivermaxCore/Private/RivermaxPTPUtils.cpp
+++ b/Engine/Plugins/VirtualProduction/Rivermax/RivermaxCore/Source/RivermaxCore/Private/RivermaxPTPUtils.cpp
@@ -25,98 +25,43 @@ namespace UE::RivermaxCore
 		// Based on ST2059. Formulae from PTP time to next alignment point
 		// NextAlignmentPoint = Floor( (CurrentFrameNumber + 1) * IntervalNs)
 
-		// Example for frame F=40113798005 at 23.97
-		// uint64 T = F / 24000 = 1671408
-		// uint64 Tr =F % 24000 = 6005
-		// uint64 TT = T * 1001 = 1673079408
-		// uint64 TTr = 6005 * 1001 = 6011005
-		// uint64 TTrr = TTr * Scale / 24000 = 250458541666
-		// return TT * Scale + TTrr = 1673079658458541666
-
 		const uint64 Nanoscale = 1E9;
 
-		const uint64 Intermediate = InFrameNumber / InRate.Numerator;
-		const uint64 IntermediateRemainder = InFrameNumber % InRate.Numerator;
-		const uint64 NextAlignment = Intermediate * InRate.Denominator;
-		const uint64 NextAlignementRemainderInterm = (IntermediateRemainder * InRate.Denominator);
-		const uint64 NextAlignmentRemainder = (NextAlignementRemainderInterm * Nanoscale) / InRate.Numerator;
+		const uint64 NextAlignment = (InFrameNumber / InRate.Numerator) * InRate.Denominator;
+		const uint64 NextAlignmentRemainderInterm = (InFrameNumber % InRate.Numerator) * InRate.Denominator;
+		const uint64 NextAlignmentRemainder = (NextAlignmentRemainderInterm * Nanoscale) / InRate.Numerator;
 
-		return (NextAlignment * Nanoscale) + NextAlignmentRemainder;
+		// We ceil to next nanosecond in case of any remainding. i.e 333333.33333 nanosec => 333334
+		const uint64 OneNanoUp = (NextAlignmentRemainderInterm * Nanoscale) % InRate.Numerator > 0 ? 1 : 0;
+		return (NextAlignment * Nanoscale) + NextAlignmentRemainder + OneNanoUp;
 	}
 
 	uint64 GetFrameNumber(const uint64 InPTPTimeNanosec, const FFrameRate& InRate)
 	{
 		const uint64 NanoScale = 1E9;
 
-		// Util struct returned for intermediate calculation of frame number
-		struct FResult
-		{
-			uint64 FrameCount = 0;
-			double Fraction = 0.0;
-		};
-
-		const auto CalculateFrameNumberFunc = [InRate](const uint64 PTPTime, uint64 Scale) -> FResult
-		{
-			// Example for PTP time 1673079658458541667 at 23.97
-			
-			// Seconds part (1673079658). Scale = 1. Remainder = 0.
-			// uint64 F = Time / 1001 = 1671408
-			// uint64 Fr = Time % 1001 = 250
-			// uint64 FF = F * 24000 = 40113792000
-			// uint64 FFr = (Fr * 24000) / 1001 = 5994
-			// uint64 FFrr =(Fr * 24000) % 1001) = 6
-			// double Tr = (FFrr / (double)1001.0 = 0.005994005994
-			// uint64 Frame = (FF + FFr) / Scale = 40113797994
-			// double FractionalFrame = Frac((FF + FFr) / Scale) + Tr / Scale = 0.005994005994
-			// ret {Frame, Tr} = {40113797994, 0.005994005994}
-
-			// Nanoseconds part (458541667). Scale = 1E9. Remainder = 0.005994005994 * 1E9 = 5994005
-			// uint64 F = Time / 1001 = 458083
-			// uint64 Fr = Time % 1001 = 584
-			// uint64 FF = F * 24000 = 10993992000
-			// uint64 FFr = (Fr * 24000) / 1001 = 14001
-			// uint64 FFrr =(Fr * 24000) % 1001) = 999
-			// double Tr = (FFrr / (double)1001.0 = 0.998001998001998002
-			// uint64 Frame = (FF + FFr) / Scale = 10
-			// ret {Frame, Tr} = {10, .994006001998002}
-
-			// Full frames = 40113797994 + 10 =  40113798004
-			// Fractional frames = 1.0000000079920019 ( 0.005994005994 + 0.994006001998002 )
-			// We round to the 1/1000000th of a frame to cover cases where fractional frame is 0.999999
-			// Final result= 40113798004 + 1 = 40113798005
-
-
-			const uint64 FrameIntermediate = PTPTime / InRate.Denominator;
-			const uint64 FrameRemainderIntermediate = PTPTime % InRate.Denominator;
-			const uint64 FrameInteger = FrameIntermediate * InRate.Numerator;
-			const uint64 FrameIntegerRemainder = (FrameRemainderIntermediate * InRate.Numerator) / InRate.Denominator;
-			const double TrailingRemainder = ((FrameRemainderIntermediate * InRate.Numerator) % InRate.Denominator) / (double)InRate.Denominator;
-			const uint64 FrameNumber = (FrameInteger + FrameIntegerRemainder) / Scale;
-
-			const double FractionalFrame  = FMath::Frac((FrameInteger + FrameIntegerRemainder) / (double)Scale) + TrailingRemainder / Scale;
-			return FResult{ FrameNumber, FractionalFrame };
-		};
-
 		// Split incoming time in nanoseconds in two parts : seconds and nanoseconds
 		// Even after dividing by frame rate denominator (i.e 1001), we might overflow uint64 when multiplying by frame rate numerator (i.e. 24000)
 		const uint64 PTPSeconds = InPTPTimeNanosec / NanoScale;
 		const uint64 PTPNanoSeconds = InPTPTimeNanosec % NanoScale;
+
+		// Truncated frames from the seconds portion
+		const uint64 FSec = (PTPSeconds * InRate.Numerator) / InRate.Denominator;
+
+		// We're missing the fractional frames from the seconds portion
+		// ((PTPSeconds * InRate.Numerator) % InRate.Denominator) / InRate.Denominator;
+
+		// Truncated frames from the nanoseconds portion
+		const uint64 FNanoSec = (PTPNanoSeconds * InRate.Numerator) / (InRate.Denominator * NanoScale);
+
+		// We're missing the fractional frames from the nanoseconds portion
+		// ((PTPNanoSeconds * InRate.Numerator) % (InRate.Denominator * NanoScale)) / ((InRate.Denominator * NanoScale))
 		
-		// Proceed with seconds part. No scaling
-		const FResult FrameNumber1 = CalculateFrameNumberFunc(PTPSeconds, 1);
-
-		// Proceed with nanoseconds using a scale of 1E9 to get actual frame numbers
-		const FResult FrameNumber2 = CalculateFrameNumberFunc(PTPNanoSeconds, NanoScale);
-
-		// Add up full frame number
-		const uint64 FrameNumber = FrameNumber1.FrameCount + FrameNumber2.FrameCount;
-
-		// Add up fractional frame number from intermediate calculation
-		// and round to the 1/100000th of a frame to avoid returning the previous time when the incoming time
-		// is almost exactly on alignment
-		const double FramePrecision = 100000; 
-		const double FractionFrameNumber = FrameNumber1.Fraction + FrameNumber2.Fraction;
-		const uint64 TrailingFrameNumber = FMath::RoundToDouble(FractionFrameNumber * FramePrecision) / FramePrecision;
-		return FrameNumber + TrailingFrameNumber;
+		// Add the fractional frames from both potions
+		uint64 FractionalFrames = NanoScale * ((PTPSeconds * InRate.Numerator) % InRate.Denominator);
+		FractionalFrames += (PTPNanoSeconds * InRate.Numerator) % (InRate.Denominator * NanoScale);
+		FractionalFrames = FractionalFrames / (InRate.Denominator * NanoScale);
+
+		return FSec + FNanoSec + FractionalFrames;
 	}
 }
\ No newline at end of file
diff --git a/Engine/Plugins/VirtualProduction/Rivermax/RivermaxCore/Source/RivermaxCore/Private/RivermaxUtils.cpp b/Engine/Plugins/VirtualProduction/Rivermax/RivermaxCore/Source/RivermaxCore/Private/RivermaxUtils.cpp
--- a/Engine/Plugins/VirtualProduction/Rivermax/RivermaxCore/Source/RivermaxCore/Private/RivermaxUtils.cpp
+++ b/Engine/Plugins/VirtualProduction/Rivermax/RivermaxCore/Source/RivermaxCore/Private/RivermaxUtils.cpp
@@ -7,50 +7,87 @@
 
 namespace UE::RivermaxCore::Private::Utils
 {
-	FString PixelFormatToSamplingDesc(ERivermaxOutputPixelFormat PixelFormat)
+	FString PixelFormatToSamplingDesc(ESamplingType SamplingType)
 	{
-		switch (PixelFormat)
+		switch (SamplingType)
 		{
-			case ERivermaxOutputPixelFormat::RMAX_10BIT_RGB:
-			case ERivermaxOutputPixelFormat::RMAX_8BIT_RGB:
-			case ERivermaxOutputPixelFormat::RMAX_16F_RGB:
+			case ESamplingType::RGB_8bit:
+			case ESamplingType::RGB_10bit:
+			case ESamplingType::RGB_12bit:
+			case ESamplingType::RGB_16bit:
+			case ESamplingType::RGB_16bitFloat:
 			{
 				return FString(TEXT("RGB"));
 			}
-			case ERivermaxOutputPixelFormat::RMAX_10BIT_YCBCR:
-			case ERivermaxOutputPixelFormat::RMAX_8BIT_YCBCR:
-			default:
+			case ESamplingType::YUV422_8bit:
+			case ESamplingType::YUV422_10bit:
+			case ESamplingType::YUV422_12bit:
+			case ESamplingType::YUV422_16bit:
+			case ESamplingType::YUV422_16bitFloat:
 			{
 				return FString(TEXT("YCbCr-4:2:2"));
 			}
-		}
-	}
-
-	FString PixelFormatToBitDepth(ERivermaxOutputPixelFormat PixelFormat)
-	{
-		switch (PixelFormat)
-		{
-			case ERivermaxOutputPixelFormat::RMAX_10BIT_RGB:
-			case ERivermaxOutputPixelFormat::RMAX_10BIT_YCBCR:
-			{
-				return FString(TEXT("10"));
-			}
-			case ERivermaxOutputPixelFormat::RMAX_16F_RGB:
+			case ESamplingType::YUV444_8bit:
+			case ESamplingType::YUV444_10bit:
+			case ESamplingType::YUV444_12bit:
+			case ESamplingType::YUV444_16bit:
+			case ESamplingType::YUV444_16bitFloat:
 			{
-				return FString(TEXT("16f"));
+				return FString(TEXT("YCbCr-4:4:4"));
 			}
-			case ERivermaxOutputPixelFormat::RMAX_8BIT_RGB:
-			case ERivermaxOutputPixelFormat::RMAX_8BIT_YCBCR:
 			default:
 			{
-				return FString(TEXT("8"));
+				checkNoEntry();
+				return FString();
 			}
 		}
 	}
 
+	FString PixelFormatToBitDepth(ESamplingType SamplingType)
+	{
+		switch (SamplingType)
+		{
+		case ESamplingType::RGB_8bit:
+		case ESamplingType::YUV422_8bit:
+		case ESamplingType::YUV444_8bit:
+		{
+			return FString(TEXT("8"));
+		}
+		case ESamplingType::RGB_10bit:
+		case ESamplingType::YUV422_10bit:
+		case ESamplingType::YUV444_10bit:
+		{
+			return FString(TEXT("10"));
+		}
+		case ESamplingType::RGB_12bit:
+		case ESamplingType::YUV422_12bit:
+		case ESamplingType::YUV444_12bit:
+		{
+			return FString(TEXT("12"));
+		}
+		case ESamplingType::RGB_16bit:
+		case ESamplingType::YUV422_16bit:
+		case ESamplingType::YUV444_16bit:
+		{
+			return FString(TEXT("16"));
+		}
+		case ESamplingType::RGB_16bitFloat:
+		case ESamplingType::YUV422_16bitFloat:
+		case ESamplingType::YUV444_16bitFloat:
+		{
+			return FString(TEXT("16f"));
+		}
+		default:
+		{
+			checkNoEntry();
+			return FString();
+		}
+		}
+	}
+
 	void StreamOptionsToSDPDescription(const UE::RivermaxCore::FRivermaxStreamOptions& Options, FAnsiStringBuilderBase& OutSDPDescription)
 	{
-		// Basic SDP string creation from a set of options. At some point, having a proper SDP loader / creator.
+		// Basic SDP string creation from a set of options. At some point, having a proper SDP loader / creator would be useful.
 		// Refer to https://datatracker.ietf.org/doc/html/rfc4570
 
 		FString FrameRateDescription;
@@ -64,11 +101,12 @@ namespace UE::RivermaxCore::Private::Utils
 		}
 
 
-
+		constexpr int32 MulticastTTL = 64;
 		OutSDPDescription.Appendf("v=0\n");
 		OutSDPDescription.Appendf("s=SMPTE ST2110 20 streams\n");
+		OutSDPDescription.Appendf("t=0 0\n");
 		OutSDPDescription.Appendf("m=video %d RTP/AVP 96\n", Options.Port);
-		OutSDPDescription.Appendf("c=IN IP4 %S/64\n", *Options.StreamAddress);
+		OutSDPDescription.Appendf("c=IN IP4 %S/%d\n", *Options.StreamAddress, MulticastTTL);
 		OutSDPDescription.Appendf("a=source-filter: incl IN IP4 %S %S\n", *Options.StreamAddress, *Options.InterfaceAddress);
 		OutSDPDescription.Appendf("a=rtpmap:96 raw/90000\n");
 		OutSDPDescription.Appendf("a=fmtp: 96 sampling=%S; width=%d; height=%d; exactframerate=%S; depth=%S; TCS=SDR; colorimetry=BT709; PM=2110GPM; SSN=ST2110-20:2017; TP=2110TPN;\n"
@@ -77,8 +115,5 @@ namespace UE::RivermaxCore::Private::Utils
 			, Options.Resolution.Y
 			, *FrameRateDescription
 			, *PixelFormatToBitDepth(Options.PixelFormat));
-
-		OutSDPDescription.Appendf("a=mediaclk:direct=0\n");
-		OutSDPDescription.Appendf("a=mid:VID");
 	}
 }
\ No newline at end of file
diff --git a/Engine/Plugins/VirtualProduction/Rivermax/RivermaxCore/Source/RivermaxCore/Private/RivermaxUtils.h b/Engine/Plugins/VirtualProduction/Rivermax/RivermaxCore/Source/RivermaxCore/Private/RivermaxUtils.h
--- a/Engine/Plugins/VirtualProduction/Rivermax/RivermaxCore/Source/RivermaxCore/Private/RivermaxUtils.h
+++ b/Engine/Plugins/VirtualProduction/Rivermax/RivermaxCore/Source/RivermaxCore/Private/RivermaxUtils.h
@@ -4,6 +4,8 @@
 
 #include "CoreMinimal.h"
 
+#include "RivermaxTypes.h"
+
 namespace UE::RivermaxCore
 {
 	struct FRivermaxStreamOptions;
@@ -12,12 +14,15 @@ namespace UE::RivermaxCore
 namespace UE::RivermaxCore::Private::Utils
 {	
 	/** Various constants used for stream initialization */
-	static constexpr uint32 BytesPerHeader = 20;
-	static constexpr uint32 BytesPerPacket = 1200;
+	static constexpr uint32 RTPHeaderSize = 20;
 	static constexpr uint32 FullHDHeight = 1080;
 	static constexpr uint32 FullHDWidth = 1920;
-	static constexpr uint32 MaxPayloadSize = 1280;
-	static constexpr uint32 VideoTROModification = 0; // Todo : Investigate what TRO modification is
+
+	/** Maximum payload we can send based on UDP max size and RTP header.  */
+	static constexpr uint32 MaxPayloadSize = 1440;
+	
+	/** Smallest payload size to bound our search for a payload that can be equal across a line */
+	static constexpr uint32 MinPayloadSize = 950;
 
 	/** Convert a set of streaming option to its SDP description. Currently only support video type. */
 	void StreamOptionsToSDPDescription(const UE::RivermaxCore::FRivermaxStreamOptions& Options, FAnsiStringBuilderBase& OutSDPDescription);
diff --git a/Engine/Plugins/VirtualProduction/Rivermax/RivermaxCore/Source/RivermaxCore/Private/Tests/PTPTest.cpp b/Engine/Plugins/VirtualProduction/Rivermax/RivermaxCore/Source/RivermaxCore/Private/Tests/PTPTest.cpp
--- a/Engine/Plugins/VirtualProduction/Rivermax/RivermaxCore/Source/RivermaxCore/Private/Tests/PTPTest.cpp
+++ b/Engine/Plugins/VirtualProduction/Rivermax/RivermaxCore/Source/RivermaxCore/Private/Tests/PTPTest.cpp
@@ -31,53 +31,55 @@ bool FPTPTest::RunTest(const FString& Parameters)
 
 	// Add couple of test points at different frame rate to test PTP operation
 
-	TestPoints.Add({ FFrameRate(24,1), 0, 0, 41666666 });
-	TestPoints.Add({ FFrameRate(24,1), 42000000, 1, 83333333 });
-	TestPoints.Add({ FFrameRate(24,1), 1671408000000000000, 40113792000, 1671408000041666666 });
-	TestPoints.Add({ FFrameRate(24,1), 1671408000041666666, 40113792001, 1671408000083333333 }); 
-	TestPoints.Add({ FFrameRate(24,1), 1671408000083333333, 40113792002, 1671408000125000000 });
-	TestPoints.Add({ FFrameRate(24,1), 1671408000020000000, 40113792000, 1671408000041666666 }); 
-	TestPoints.Add({ FFrameRate(24,1), 1671408000040000000, 40113792000, 1671408000041666666 });
+	TestPoints.Add({ FFrameRate(24,1), 0, 0, 41666667 });
+	TestPoints.Add({ FFrameRate(24,1), 42000000, 1, 83333334 });
+	TestPoints.Add({ FFrameRate(24,1), 1671408000000000000, 40113792000, 1671408000041666667 });
+	TestPoints.Add({ FFrameRate(24,1), 1671408000041666666, 40113792000, 1671408000041666667 });
+	TestPoints.Add({ FFrameRate(24,1), 1671408000041666667, 40113792001, 1671408000083333334 });
+	TestPoints.Add({ FFrameRate(24,1), 1671408000083333334, 40113792002, 1671408000125000000 });
+	TestPoints.Add({ FFrameRate(24,1), 1671408000020000000, 40113792000, 1671408000041666667 }); 
+	TestPoints.Add({ FFrameRate(24,1), 1671408000040000000, 40113792000, 1671408000041666667 });
 	TestPoints.Add({ FFrameRate(24,1), 1673079658458541667, 40153911803, 1673079658500000000 });
-	TestPoints.Add({ FFrameRate(24,1), 4099680000000000000, 98392320000, 4099680000041666666 });
-	TestPoints.Add({ FFrameRate(24,1), 4099680000041666666, 98392320001, 4099680000083333333 }); // Large nano second test. Approx 130 years from epoch
-
-	TestPoints.Add({ FFrameRate(24000,1001), 0, 0, 41708333 }); 
-	TestPoints.Add({ FFrameRate(24000,1001), 42000000, 1, 83416666 });
-	TestPoints.Add({ FFrameRate(24000,1001), 1673079408000000000, 40113792000, 1673079408041708333 });
-	TestPoints.Add({ FFrameRate(24000,1001), 1673079408020000000, 40113792000, 1673079408041708333 });  
-	TestPoints.Add({ FFrameRate(24000,1001), 1673079408040000000, 40113792000, 1673079408041708333 });  
-	TestPoints.Add({ FFrameRate(24000,1001), 1673079408041708333, 40113792001, 1673079408083416666 }); // Right on alignment. Without frame number round up, it would return same input time
-	TestPoints.Add({ FFrameRate(24000,1001), 1673079408083416666, 40113792002, 1673079408125125000 }); // Right on alignment.
-	TestPoints.Add({ FFrameRate(24000,1001), 1673079408125125000, 40113792003, 1673079408166833333 }); // Right on alignment.
-	TestPoints.Add({ FFrameRate(24000,1001), 1673079408125124999, 40113792003, 1673079408166833333 }); // Right on alignment.
-	TestPoints.Add({ FFrameRate(24000,1001), 1673079408166833333, 40113792004, 1673079408208541666 });
-	TestPoints.Add({ FFrameRate(24000,1001), 1673079408208541666, 40113792005, 1673079408250250000 });
-	TestPoints.Add({ FFrameRate(24000,1001), 1673079408250249998, 40113792006, 1673079408291958333 });
-	TestPoints.Add({ FFrameRate(24000,1001), 1673079408250250000, 40113792006, 1673079408291958333 });
-	TestPoints.Add({ FFrameRate(24000,1001), 1673079408291958333, 40113792007, 1673079408333666666 });
-	TestPoints.Add({ FFrameRate(24000,1001), 1673079408333666666, 40113792008, 1673079408375375000 });
+	TestPoints.Add({ FFrameRate(24,1), 4099680000000000000, 98392320000, 4099680000041666667 });
+	TestPoints.Add({ FFrameRate(24,1), 4099680000041666667, 98392320001, 4099680000083333334 }); // Large nano second test. Approx 130 years from epoch
+
+	TestPoints.Add({ FFrameRate(24000,1001), 0, 0, 41708334 }); 
+	TestPoints.Add({ FFrameRate(24000,1001), 42000000, 1, 83416667 });
+	TestPoints.Add({ FFrameRate(24000,1001), 1673079408000000000, 40113792000, 1673079408041708334 });
+	TestPoints.Add({ FFrameRate(24000,1001), 1673079408020000000, 40113792000, 1673079408041708334 });  
+	TestPoints.Add({ FFrameRate(24000,1001), 1673079408040000000, 40113792000, 1673079408041708334 });  
+	TestPoints.Add({ FFrameRate(24000,1001), 1673079408041708334, 40113792001, 1673079408083416667 }); 
+	TestPoints.Add({ FFrameRate(24000,1001), 1673079408083416667, 40113792002, 1673079408125125000 }); 
+	TestPoints.Add({ FFrameRate(24000,1001), 1673079408125124999, 40113792002, 1673079408125125000 });
+	TestPoints.Add({ FFrameRate(24000,1001), 1673079408125125000, 40113792003, 1673079408166833334 });
+	TestPoints.Add({ FFrameRate(24000,1001), 1673079408166833334, 40113792004, 1673079408208541667 });
+	TestPoints.Add({ FFrameRate(24000,1001), 1673079408208541667, 40113792005, 1673079408250250000 });
+	TestPoints.Add({ FFrameRate(24000,1001), 1673079408250249998, 40113792005, 1673079408250250000 });
+	TestPoints.Add({ FFrameRate(24000,1001), 1673079408250250000, 40113792006, 1673079408291958334 });
+	TestPoints.Add({ FFrameRate(24000,1001), 1673079408291958334, 40113792007, 1673079408333666667 });
+	TestPoints.Add({ FFrameRate(24000,1001), 1673079408333666667, 40113792008, 1673079408375375000 });
 	TestPoints.Add({ FFrameRate(24000,1001), 4099680000000000000, 98294025974, 4099680000040625000 });
-	TestPoints.Add({ FFrameRate(24000,1001), 4099680000040625000, 98294025975, 4099680000082333333 });
-	TestPoints.Add({ FFrameRate(24000,1001), 4099680000082333333, 98294025976, 4099680000124041666 });
-
-	TestPoints.Add({ FFrameRate(60,1), 0, 0, 16666666 });
-	TestPoints.Add({ FFrameRate(60,1), 26666666, 1, 33333333 });
-	TestPoints.Add({ FFrameRate(60,1), 2674252800000000000, 160455168000, 2674252800016666666 }); 
-	TestPoints.Add({ FFrameRate(60,1), 2674252800008000000, 160455168000, 2674252800016666666 });
-	TestPoints.Add({ FFrameRate(60,1), 2674252800010000000, 160455168000, 2674252800016666666 });
-	TestPoints.Add({ FFrameRate(60,1), 1673079658458541667, 100384779507, 1673079658466666666 });
-	TestPoints.Add({ FFrameRate(60,1), 1673079658466666666, 100384779508, 1673079658483333333 });
-	TestPoints.Add({ FFrameRate(60,1), 4099680000000000000, 245980800000, 4099680000016666666 });
-	TestPoints.Add({ FFrameRate(60,1), 4099680000016666666, 245980800001, 4099680000033333333 });
-
-	TestPoints.Add({ FFrameRate(60000,1001), 0, 0, 16683333 });
-	TestPoints.Add({ FFrameRate(60000,1001), 26666666, 1, 33366666 });
-	TestPoints.Add({ FFrameRate(60000,1001), 2674252800000000000, 160294873126, 2674252800002116666 });
+	TestPoints.Add({ FFrameRate(24000,1001), 4099680000040625000, 98294025975, 4099680000082333334 });
+	TestPoints.Add({ FFrameRate(24000,1001), 4099680000082333334, 98294025976, 4099680000124041667 });
+
+	TestPoints.Add({ FFrameRate(60,1), 0, 0, 16666667 });
+	TestPoints.Add({ FFrameRate(60,1), 26666666, 1, 33333334 });
+	TestPoints.Add({ FFrameRate(60,1), 2674252800000000000, 160455168000, 2674252800016666667 }); 
+	TestPoints.Add({ FFrameRate(60,1), 2674252800008000000, 160455168000, 2674252800016666667 });
+	TestPoints.Add({ FFrameRate(60,1), 2674252800010000000, 160455168000, 2674252800016666667 });
+	TestPoints.Add({ FFrameRate(60,1), 1673079658458541667, 100384779507, 1673079658466666667 });
+	TestPoints.Add({ FFrameRate(60,1), 1673079658466666667, 100384779508, 1673079658483333334 });
+	TestPoints.Add({ FFrameRate(60,1), 4099680000000000000, 245980800000, 4099680000016666667 });
+	TestPoints.Add({ FFrameRate(60,1), 4099680000016666666, 245980800000, 4099680000016666667 });
+	TestPoints.Add({ FFrameRate(60,1), 4099680000016666667, 245980800001, 4099680000033333334 });
+
+	TestPoints.Add({ FFrameRate(60000,1001), 0, 0, 16683334 });
+	TestPoints.Add({ FFrameRate(60000,1001), 26666666, 1, 33366667 });
+	TestPoints.Add({ FFrameRate(60000,1001), 2674252800000000000, 160294873126, 2674252800002116667 });
 	TestPoints.Add({ FFrameRate(60000,1001), 2674252800008000000, 160294873127, 2674252800018800000 });
 	TestPoints.Add({ FFrameRate(60000,1001), 2674252800010000000, 160294873127, 2674252800018800000 });
-	TestPoints.Add({ FFrameRate(60000,1001), 1673079658458541667, 100284495012, 1673079658466883333 });
-	TestPoints.Add({ FFrameRate(60000,1001), 1673079658466883333, 100284495013, 1673079658483566666 });
+	TestPoints.Add({ FFrameRate(60000,1001), 1673079658458541667, 100284495012, 1673079658466883334 });
+	TestPoints.Add({ FFrameRate(60000,1001), 1673079658466883334, 100284495013, 1673079658483566667 });
 	TestPoints.Add({ FFrameRate(60000,1001), 4099680000000000000, 245735064935, 4099680000015600000 });
 	
 
diff --git a/Engine/Plugins/VirtualProduction/Rivermax/RivermaxCore/Source/RivermaxCore/Public/RivermaxFormats.h b/Engine/Plugins/VirtualProduction/Rivermax/RivermaxCore/Source/RivermaxCore/Public/RivermaxFormats.h
new file mode 100644
--- /dev/null
+++ b/Engine/Plugins/VirtualProduction/Rivermax/RivermaxCore/Source/RivermaxCore/Public/RivermaxFormats.h
@@ -0,0 +1,58 @@
+// Copyright Epic Games, Inc. All Rights Reserved.
+
+#pragma once
+
+#include "CoreMinimal.h"
+
+#include "Misc/FrameRate.h"
+
+namespace UE::RivermaxCore
+{
+	enum class RIVERMAXCORE_API ESamplingType : uint8
+	{
+		YUV422_8bit
+		, YUV422_10bit
+		, YUV422_12bit
+		, YUV422_16bit
+		, YUV422_16bitFloat
+		, YUV444_8bit
+		, YUV444_10bit
+		, YUV444_12bit
+		, YUV444_16bit
+		, YUV444_16bitFloat
+		, RGB_8bit
+		, RGB_10bit
+		, RGB_12bit
+		, RGB_16bit
+		, RGB_16bitFloat
+	};
+
+	struct RIVERMAXCORE_API FVideoFormatInfo
+	{
+		/** Sampling type of this format */
+		ESamplingType Sampling = ESamplingType::RGB_10bit;
+
+		/** Bit depth for this format (i.e 8,10,12,16) */
+		uint16 BitDepth = 0;
+
+		/** Pixel Group (pgroup) size in bytes */
+		uint16 PixelGroupSize = 0;
+
+		/** Pixel Group coverage (pgroup coverage) in pixels */
+		uint16 PixelGroupCoverage = 0;
+	};
+
+	struct RIVERMAXCORE_API FStandardVideoFormat
+	{
+		static FVideoFormatInfo GetVideoFormatInfo(ESamplingType SamplingType);
+		static bool IsRGB(ESamplingType SamplingType);
+		static bool IsYUV(ESamplingType SamplingType);
+		static bool IsYUV422(ESamplingType SamplingType);
+		static bool IsYUV444(ESamplingType SamplingType);
+
+	private:
+		static const TArray<FVideoFormatInfo> AllSamplingTypes;
+	};
+}
+
+
diff --git a/Engine/Plugins/VirtualProduction/Rivermax/RivermaxCore/Source/RivermaxCore/Public/RivermaxTypes.h b/Engine/Plugins/VirtualProduction/Rivermax/RivermaxCore/Source/RivermaxCore/Public/RivermaxTypes.h
--- a/Engine/Plugins/VirtualProduction/Rivermax/RivermaxCore/Source/RivermaxCore/Public/RivermaxTypes.h
+++ b/Engine/Plugins/VirtualProduction/Rivermax/RivermaxCore/Source/RivermaxCore/Public/RivermaxTypes.h
@@ -5,18 +5,10 @@
 #include "CoreMinimal.h"
 
 #include "Misc/FrameRate.h"
+#include "RivermaxFormats.h"
 
 namespace UE::RivermaxCore
 {
-	enum class RIVERMAXCORE_API ERivermaxOutputPixelFormat
-	{
-		RMAX_8BIT_YCBCR,
-		RMAX_8BIT_RGB,
-		RMAX_10BIT_YCBCR,
-		RMAX_10BIT_RGB,
-		RMAX_16F_RGB,
-	};
-
 	struct RIVERMAXCORE_API FRivermaxStreamOptions
 	{
 		/** Desired stream resolution */
@@ -35,13 +27,13 @@ namespace UE::RivermaxCore
 		uint32 Port = 50000;
 
 		/** Desired stream pixel format */
-		ERivermaxOutputPixelFormat PixelFormat = ERivermaxOutputPixelFormat::RMAX_8BIT_YCBCR;
+		ESamplingType PixelFormat = ESamplingType::RGB_10bit;
 
 		/** Sample count to buffer. */
 		int32 NumberOfBuffers = 2;
 
-		/** Stride in bytes for a row */
-		uint32 Stride = 0;
+		/** Resolution aligning with pgroup of sampling type */
+		FIntPoint AlignedResolution = FIntPoint::ZeroValue;
 	};
 
 	enum class RIVERMAXCORE_API ERivermaxStreamType : uint8
diff --git a/Engine/Plugins/VirtualProduction/Rivermax/RivermaxMedia/Source/RivermaxMedia/Private/RivermaxMediaCapture.cpp b/Engine/Plugins/VirtualProduction/Rivermax/RivermaxMedia/Source/RivermaxMedia/Private/RivermaxMediaCapture.cpp
--- a/Engine/Plugins/VirtualProduction/Rivermax/RivermaxMedia/Source/RivermaxMedia/Private/RivermaxMediaCapture.cpp
+++ b/Engine/Plugins/VirtualProduction/Rivermax/RivermaxMedia/Source/RivermaxMedia/Private/RivermaxMediaCapture.cpp
@@ -8,6 +8,7 @@
 #include "RenderGraphUtils.h"
 #include "RivermaxMediaLog.h"
 #include "RivermaxMediaOutput.h"
+#include "RivermaxMediaUtils.h"
 #include "RivermaxShaders.h"
 #include "RivermaxTypes.h"
 #include "Slate/SceneViewport.h"
@@ -49,50 +50,54 @@ namespace RivermaxMediaCaptureAnalytics
 
 namespace UE::RivermaxMediaCaptureUtil
 {
+	FIntPoint GetAlignedResolution(const UE::RivermaxCore::FVideoFormatInfo& InFormatInfo, const FIntPoint& InSize)
+	{
+		// We need to adjust horizontal resolution if it's not aligned to pgroup pixel coverage
+		// RTP header has to describe a full pgroup even if it's outside of effective resolution.
+		const uint32 PixelAlignment = InFormatInfo.PixelGroupCoverage;
+		const uint32 AlignedHorizontalResolution = InSize.X % PixelAlignment ? InSize.X + (PixelAlignment - InSize.X % PixelAlignment) : InSize.X;
+		return FIntPoint(AlignedHorizontalResolution, InSize.Y);
+	}
+
 	void GetOutputEncodingInfo(ERivermaxMediaOutputPixelFormat InPixelFormat, const FIntPoint& InSize, uint32& OutBytesPerElement, uint32& ElementsPerRow)
 	{
+		using namespace UE::RivermaxCore;
+
+		const ESamplingType SamplingType = UE::RivermaxMediaUtils::Private::MediaOutputPixelFormatToRivermaxSamplingType(InPixelFormat);
+		const FVideoFormatInfo Info = FStandardVideoFormat::GetVideoFormatInfo(SamplingType);
+
+		// Compute horizontal byte count (stride) of aligned resolution
+		const uint32 PixelAlignment = Info.PixelGroupCoverage;
+		const FIntPoint AlignedResolution = GetAlignedResolution(Info, InSize);
+		const uint32 PixelGroupCount = AlignedResolution.X / PixelAlignment;
+		const uint32 HorizontalByteCount = PixelGroupCount * Info.PixelGroupSize;
+
 		switch (InPixelFormat)
 		{
-		case ERivermaxMediaOutputPixelFormat::PF_10BIT_YCBCR_422:
+		case ERivermaxMediaOutputPixelFormat::PF_10BIT_YUV422:
 		{
-			// 2 RGB pixels per YUV pixels and 4 YUV pixels per entry in the buffer
-			constexpr uint32 RGBPixelsPerYUVPixels = 2;
-			const uint32 EncodedHorizontalResolution = InSize.X / RGBPixelsPerYUVPixels;
-			constexpr uint32 BytesPerYUVPixels = 5; //5 bytes per pixels for YUV422 10bits
-			const uint32 HorizontalByteCount = EncodedHorizontalResolution * BytesPerYUVPixels;
 			OutBytesPerElement = sizeof(UE::RivermaxShaders::FRGBToYUV10Bit422LittleEndianCS::FYUV10Bit422LEBuffer);
-			ElementsPerRow = HorizontalByteCount / OutBytesPerElement;
 			break;
 		}
 		case ERivermaxMediaOutputPixelFormat::PF_8BIT_RGB:
 		{
-			// 2 RGB pixels per YUV pixels and 4 YUV pixels per entry in the buffer
-			constexpr uint32 BytesPerRGBPixels = 3; //3 bytes per pixels RGB8
-			const uint32 HorizontalByteCount = InSize.X * BytesPerRGBPixels;
 			OutBytesPerElement = sizeof(UE::RivermaxShaders::FRGBToRGB8BitCS::FRGB8BitBuffer);
-			ElementsPerRow = HorizontalByteCount / OutBytesPerElement;
 			break;
 		}
 		case ERivermaxMediaOutputPixelFormat::PF_10BIT_RGB:
 		{
-			// 2 RGB pixels per YUV pixels and 4 YUV pixels per entry in the buffer
-			constexpr uint32 PixelPerGroup = 4;
-			constexpr uint32 BytesPerGroup = (PixelPerGroup * 3 * 10) / 8;
-			const uint32 HorizontalByteCount = (InSize.X / PixelPerGroup) * BytesPerGroup;
 			OutBytesPerElement = sizeof(UE::RivermaxShaders::FRGBToRGB10BitCS::FRGB10BitBuffer);
-			ElementsPerRow = HorizontalByteCount / OutBytesPerElement;
 			break;
 		}
 		case ERivermaxMediaOutputPixelFormat::PF_FLOAT16_RGB:
 		{
-			constexpr uint32 PixelPerGroup = 2;
-			constexpr uint32 BytesPerGroup = (PixelPerGroup * 3 * 16) / 8;
-			const uint32 HorizontalByteCount = (InSize.X / PixelPerGroup) * BytesPerGroup;
-			OutBytesPerElement = BytesPerGroup;
-			ElementsPerRow = HorizontalByteCount / OutBytesPerElement;
+			OutBytesPerElement = sizeof(UE::RivermaxShaders::FRGBToRGB16fCS::FRGB16fBuffer);
 			break;
 		}
 		}
+
+		// Shader encoding might not align with pixel group size so we need to have enough elements to represent the last pixel group
+		ElementsPerRow = FMath::CeilToInt32((float)HorizontalByteCount / OutBytesPerElement);
 	}
 }
 
@@ -210,44 +215,9 @@ bool URivermaxMediaCapture::ConfigureStream(URivermaxMediaOutput* InMediaOutput,
 	OutOptions.FrameRate = InMediaOutput->FrameRate;
 	OutOptions.NumberOfBuffers = InMediaOutput->NumberOfTextureBuffers;
 
-	switch (InMediaOutput->PixelFormat)
-	{
-	case ERivermaxMediaOutputPixelFormat::PF_8BIT_YUV:
-	{
-		OutOptions.PixelFormat = ERivermaxOutputPixelFormat::RMAX_8BIT_YCBCR;
-		OutOptions.Stride = InMediaOutput->Resolution.X / 2 * 4; // 4 bytes for a group of 2 pixels
-		break;
-	}
-	case ERivermaxMediaOutputPixelFormat::PF_10BIT_YCBCR_422:
-	{
-		OutOptions.PixelFormat = ERivermaxOutputPixelFormat::RMAX_10BIT_YCBCR;
-		OutOptions.Stride = InMediaOutput->Resolution.X / 2 * 5; // 5 bytes for a group of 2 pixels (40bits / 8 = 5)
-		break;
-	}
-	case ERivermaxMediaOutputPixelFormat::PF_8BIT_RGB:
-	{
-		OutOptions.PixelFormat = ERivermaxOutputPixelFormat::RMAX_8BIT_RGB;
-		OutOptions.Stride = InMediaOutput->Resolution.X * 3; //3 bytes per pixel
-		break;
-	}
-	case ERivermaxMediaOutputPixelFormat::PF_10BIT_RGB:
-	{
-		OutOptions.PixelFormat = ERivermaxOutputPixelFormat::RMAX_10BIT_RGB;
-		OutOptions.Stride = InMediaOutput->Resolution.X / 4 * 15; //15 bytes for a group of 4 pixels (4 * 30bits / 8 = 15)
-		break;
-	}
-	case ERivermaxMediaOutputPixelFormat::PF_FLOAT16_RGB:
-	{
-		OutOptions.PixelFormat = ERivermaxOutputPixelFormat::RMAX_16F_RGB;
-		OutOptions.Stride = InMediaOutput->Resolution.X * 6; //3x 16bits per pixel (6 bytes)
-		break;
-	}
-	default:
-	{
-		UE_LOG(LogRivermaxMedia, Error, TEXT("Desired pixel format (%s) is not a valid Rivermax pixel format"), *StaticEnum<ERivermaxMediaOutputPixelFormat>()->GetValueAsString(InMediaOutput->PixelFormat));
-		return false;
-	}
-	}
+	OutOptions.PixelFormat = UE::RivermaxMediaUtils::Private::MediaOutputPixelFormatToRivermaxSamplingType(InMediaOutput->PixelFormat);
+	const FVideoFormatInfo Info = FStandardVideoFormat::GetVideoFormatInfo(OutOptions.PixelFormat);
+	OutOptions.AlignedResolution = UE::RivermaxMediaCaptureUtil::GetAlignedResolution(Info, InMediaOutput->Resolution);
 
 	return true;
 }
@@ -266,7 +236,7 @@ void URivermaxMediaCapture::OnFrameCaptured_RenderingThread(const FCaptureBaseDa
 	NewFrame.FrameIdentifier = InBaseData.SourceFrameNumber;
 	if (RivermaxStream->PushVideoFrame(NewFrame) == false)
 	{
-		UE_LOG(LogRivermaxMedia, Warning, TEXT("Failed to pushed captured frame"));
+		UE_LOG(LogRivermaxMedia, VeryVerbose, TEXT("Failed to pushed captured frame"));
 	}
 }
 
@@ -282,7 +252,7 @@ FIntPoint URivermaxMediaCapture::GetCustomOutputSize(const FIntPoint& InSize) co
 	uint32 BytesPerElement = 0;
 	uint32 ElementsPerRow = 0;
 	UE::RivermaxMediaCaptureUtil::GetOutputEncodingInfo(RivermaxOutput->PixelFormat, InSize, BytesPerElement, ElementsPerRow);
-	return FIntPoint(ElementsPerRow, InSize.Y);;
+	return FIntPoint(ElementsPerRow, InSize.Y);
 }
 
 EMediaCaptureResourceType URivermaxMediaCapture::GetCustomOutputResourceType() const
@@ -291,7 +261,7 @@ EMediaCaptureResourceType URivermaxMediaCapture::GetCustomOutputResourceType() c
 
 	switch (RivermaxOutput->PixelFormat)
 	{
-	case ERivermaxMediaOutputPixelFormat::PF_10BIT_YCBCR_422:
+	case ERivermaxMediaOutputPixelFormat::PF_10BIT_YUV422:
 	case ERivermaxMediaOutputPixelFormat::PF_8BIT_RGB:
 	case ERivermaxMediaOutputPixelFormat::PF_10BIT_RGB:
 	case ERivermaxMediaOutputPixelFormat::PF_FLOAT16_RGB:
@@ -328,7 +298,7 @@ void URivermaxMediaCapture::OnCustomCapture_RenderingThread(FRDGBuilder& GraphBu
 	
 	switch (RivermaxOutput->PixelFormat)
 	{
-	case ERivermaxMediaOutputPixelFormat::PF_10BIT_YCBCR_422:
+	case ERivermaxMediaOutputPixelFormat::PF_10BIT_YUV422:
 	{
 		TShaderMapRef<FRGBToYUV10Bit422LittleEndianCS> ComputeShader(GlobalShaderMap);
 		FRGBToYUV10Bit422LittleEndianCS::FParameters* Parameters = ComputeShader->AllocateAndSetParameters(GraphBuilder, InSourceTexture, SourceSize, ViewRect, DesiredOutputSize, MediaShaders::RgbToYuvRec709Scaled, MediaShaders::YUVOffset10bits, bDoLinearToSRGB, OutputBuffer);
diff --git a/Engine/Plugins/VirtualProduction/Rivermax/RivermaxMedia/Source/RivermaxMedia/Private/RivermaxMediaOutput.cpp b/Engine/Plugins/VirtualProduction/Rivermax/RivermaxMedia/Source/RivermaxMedia/Private/RivermaxMediaOutput.cpp
--- a/Engine/Plugins/VirtualProduction/Rivermax/RivermaxMedia/Source/RivermaxMedia/Private/RivermaxMediaOutput.cpp
+++ b/Engine/Plugins/VirtualProduction/Rivermax/RivermaxMedia/Source/RivermaxMedia/Private/RivermaxMediaOutput.cpp
@@ -29,7 +29,7 @@ EPixelFormat URivermaxMediaOutput::GetRequestedPixelFormat() const
 	EPixelFormat Result = EPixelFormat::PF_A2B10G10R10;
 	switch (PixelFormat)
 	{
-	case ERivermaxMediaOutputPixelFormat::PF_8BIT_YUV:
+	case ERivermaxMediaOutputPixelFormat::PF_8BIT_YUV422:
 		Result = EPixelFormat::PF_B8G8R8A8; //To be updated to use compute shader
 		break;
 	default: // All pixel formats use output buffer as their output so no need for a texture format
@@ -44,10 +44,10 @@ EMediaCaptureConversionOperation URivermaxMediaOutput::GetConversionOperation(EM
 	EMediaCaptureConversionOperation Result = EMediaCaptureConversionOperation::CUSTOM;
 	switch (PixelFormat)
 	{
-	case ERivermaxMediaOutputPixelFormat::PF_8BIT_YUV:
+	case ERivermaxMediaOutputPixelFormat::PF_8BIT_YUV422:
 		Result = EMediaCaptureConversionOperation::RGBA8_TO_YUV_8BIT;
 		break;
-	case ERivermaxMediaOutputPixelFormat::PF_10BIT_YCBCR_422:
+	case ERivermaxMediaOutputPixelFormat::PF_10BIT_YUV422:
 	case ERivermaxMediaOutputPixelFormat::PF_8BIT_RGB:
 	case ERivermaxMediaOutputPixelFormat::PF_10BIT_RGB:
 	case ERivermaxMediaOutputPixelFormat::PF_FLOAT16_RGB:
diff --git a/Engine/Plugins/VirtualProduction/Rivermax/RivermaxMedia/Source/RivermaxMedia/Private/RivermaxMediaPlayer.cpp b/Engine/Plugins/VirtualProduction/Rivermax/RivermaxMedia/Source/RivermaxMedia/Private/RivermaxMediaPlayer.cpp
--- a/Engine/Plugins/VirtualProduction/Rivermax/RivermaxMedia/Source/RivermaxMedia/Private/RivermaxMediaPlayer.cpp
+++ b/Engine/Plugins/VirtualProduction/Rivermax/RivermaxMedia/Source/RivermaxMedia/Private/RivermaxMediaPlayer.cpp
@@ -13,6 +13,7 @@
 #include "RivermaxMediaOutput.h"
 #include "RivermaxMediaSourceOptions.h"
 #include "RivermaxMediaTextureSample.h"
+#include "RivermaxMediaUtils.h"
 #include "RivermaxTypes.h"
 #include "Stats/Stats2.h"
 
@@ -20,6 +21,7 @@
 #include "EngineAnalytics.h"
 #endif
 
+
 #define LOCTEXT_NAMESPACE "FRivermaxMediaPlayer"
 
 DECLARE_CYCLE_STAT(TEXT("Rivermax MediaPlayer Request frame"), STAT_Rivermax_MediaPlayer_RequestFrame, STATGROUP_Media);
@@ -67,7 +69,7 @@ namespace UE::RivermaxMedia
 		//Video related options
 		{
 			bIsSRGBInput = Options->GetMediaOption(RivermaxMediaOption::SRGBInput, bIsSRGBInput);
-			DesiredPixelFormat = (ERivermaxMediaSourePixelFormat)Options->GetMediaOption(RivermaxMediaOption::PixelFormat, (int64)ERivermaxMediaSourePixelFormat::RGB_8bit);
+			DesiredPixelFormat = (ERivermaxMediaSourcePixelFormat)Options->GetMediaOption(RivermaxMediaOption::PixelFormat, (int64)ERivermaxMediaSourcePixelFormat::RGB_8bit);
 		}
 
 		{
@@ -278,6 +280,8 @@ namespace UE::RivermaxMedia
 
 	bool FRivermaxMediaPlayer::ConfigureStream(const IMediaOptions* Options)
 	{
+		using namespace UE::RivermaxCore;
+
 		// Resolve interface address
 		IRivermaxCoreModule* Module = FModuleManager::GetModulePtr<IRivermaxCoreModule>("RivermaxCore");
 		if (Module == nullptr)
@@ -298,38 +302,13 @@ namespace UE::RivermaxMedia
 		StreamOptions.Resolution = VideoTrackFormat.Dim;
 		StreamOptions.FrameRate = VideoFrameRate;
 
-		switch (DesiredPixelFormat)
-		{
-		case ERivermaxMediaSourePixelFormat::YUV422_8bit:
-		{
-			StreamOptions.PixelFormat = ERivermaxOutputPixelFormat::RMAX_8BIT_YCBCR;
-			StreamOptions.Stride = StreamOptions.Resolution.X / 2 * 4; // 4 bytes for a group of 2 pixels
-			break;
-		}
-		case ERivermaxMediaSourePixelFormat::YUV422_10bit:
-		{
-			StreamOptions.PixelFormat = ERivermaxOutputPixelFormat::RMAX_10BIT_YCBCR;
-			StreamOptions.Stride = StreamOptions.Resolution.X / 2 * 5; // 5 bytes for a group of 2 pixels (40bits / 8 = 5)
-			break;
-		}
-		case ERivermaxMediaSourePixelFormat::RGB_8bit:
-		{
-			StreamOptions.PixelFormat = ERivermaxOutputPixelFormat::RMAX_8BIT_RGB;
-			StreamOptions.Stride = StreamOptions.Resolution.X * 3; //3 bytes per pixel
-			break;
-		}
-		case ERivermaxMediaSourePixelFormat::RGB_10bit:
-		{
-			StreamOptions.PixelFormat = ERivermaxOutputPixelFormat::RMAX_10BIT_RGB;
-			StreamOptions.Stride = StreamOptions.Resolution.X / 4 * 15; //15 bytes for a group of 4 pixels (4 * 30bits / 8 = 15)
-			break;
-		}
-		default:
-		{
-			UE_LOG(LogRivermaxMedia, Error, TEXT("Desired pixel format (%s) is not a valid Rivermax pixel format"), *StaticEnum<ERivermaxMediaSourePixelFormat>()->GetValueAsString(DesiredPixelFormat));
-			return false;
-		}
-		}
+		// We need to adjust horizontal resolution if it's not aligned to pgroup pixel coverage
+		// RTP header has to describe a full pgroup even if it's outside of effective resolution.
+		StreamOptions.PixelFormat = UE::RivermaxMediaUtils::Private::MediaSourcePixelFormatToRivermaxSamplingType(DesiredPixelFormat);
+		const FVideoFormatInfo FormatInfo = FStandardVideoFormat::GetVideoFormatInfo(StreamOptions.PixelFormat);
+		const uint32 PixelAlignment = FormatInfo.PixelGroupCoverage;
+		const uint32 AlignedHorizontalResolution = StreamOptions.Resolution.X % PixelAlignment ? StreamOptions.Resolution.X + (PixelAlignment - StreamOptions.Resolution.X % PixelAlignment) : StreamOptions.Resolution.X;
+		StreamOptions.AlignedResolution = FIntPoint(AlignedHorizontalResolution, StreamOptions.Resolution.Y);
 
 		return true;
 	}
diff --git a/Engine/Plugins/VirtualProduction/Rivermax/RivermaxMedia/Source/RivermaxMedia/Private/RivermaxMediaPlayer.h b/Engine/Plugins/VirtualProduction/Rivermax/RivermaxMedia/Source/RivermaxMedia/Private/RivermaxMediaPlayer.h
--- a/Engine/Plugins/VirtualProduction/Rivermax/RivermaxMedia/Source/RivermaxMedia/Private/RivermaxMediaPlayer.h
+++ b/Engine/Plugins/VirtualProduction/Rivermax/RivermaxMedia/Source/RivermaxMedia/Private/RivermaxMediaPlayer.h
@@ -119,7 +119,7 @@ namespace  UE::RivermaxMedia
 		std::atomic<bool> bPauseRequested;
 
 		/** Pixel format provided by media source */
-		ERivermaxMediaSourePixelFormat DesiredPixelFormat = ERivermaxMediaSourePixelFormat::RGB_10bit;
+		ERivermaxMediaSourcePixelFormat DesiredPixelFormat = ERivermaxMediaSourcePixelFormat::RGB_10bit;
 	};
 }
 
diff --git a/Engine/Plugins/VirtualProduction/Rivermax/RivermaxMedia/Source/RivermaxMedia/Private/RivermaxMediaTextureSample.cpp b/Engine/Plugins/VirtualProduction/Rivermax/RivermaxMedia/Source/RivermaxMedia/Private/RivermaxMediaTextureSample.cpp
--- a/Engine/Plugins/VirtualProduction/Rivermax/RivermaxMedia/Source/RivermaxMedia/Private/RivermaxMediaTextureSample.cpp
+++ b/Engine/Plugins/VirtualProduction/Rivermax/RivermaxMedia/Source/RivermaxMedia/Private/RivermaxMediaTextureSample.cpp
@@ -27,21 +27,21 @@ bool FRivermaxMediaTextureSample::IsOutputSrgb() const
 	return false;
 }
 
-bool FRivermaxMediaTextureSample::ConfigureSample(uint32 InWidth, uint32 InHeight, uint32 InStride, ERivermaxMediaSourePixelFormat InSampleFormat, FTimespan InTime, const FFrameRate& InFrameRate, const TOptional<FTimecode>& InTimecode, bool bInIsSRGBInput)
+bool FRivermaxMediaTextureSample::ConfigureSample(uint32 InWidth, uint32 InHeight, uint32 InStride, ERivermaxMediaSourcePixelFormat InSampleFormat, FTimespan InTime, const FFrameRate& InFrameRate, const TOptional<FTimecode>& InTimecode, bool bInIsSRGBInput)
 {
 	EMediaTextureSampleFormat VideoSampleFormat;
 	switch (InSampleFormat)
 	{
-		case ERivermaxMediaSourePixelFormat::RGB_10bit:
+		case ERivermaxMediaSourcePixelFormat::RGB_10bit:
 			// Falls through
-		case ERivermaxMediaSourePixelFormat::YUV422_10bit:
+		case ERivermaxMediaSourcePixelFormat::YUV422_10bit:
 		{
 			VideoSampleFormat = EMediaTextureSampleFormat::CharBGR10A2;
 			break;
 		}
-		case ERivermaxMediaSourePixelFormat::YUV422_8bit:
+		case ERivermaxMediaSourcePixelFormat::YUV422_8bit:
 			// Falls through
-		case ERivermaxMediaSourePixelFormat::RGB_8bit:
+		case ERivermaxMediaSourcePixelFormat::RGB_8bit:
 			// Falls through
 		default:
 		{
diff --git a/Engine/Plugins/VirtualProduction/Rivermax/RivermaxMedia/Source/RivermaxMedia/Private/RivermaxMediaTextureSample.h b/Engine/Plugins/VirtualProduction/Rivermax/RivermaxMedia/Source/RivermaxMedia/Private/RivermaxMediaTextureSample.h
--- a/Engine/Plugins/VirtualProduction/Rivermax/RivermaxMedia/Source/RivermaxMedia/Private/RivermaxMediaTextureSample.h
+++ b/Engine/Plugins/VirtualProduction/Rivermax/RivermaxMedia/Source/RivermaxMedia/Private/RivermaxMediaTextureSample.h
@@ -28,7 +28,7 @@ public:
 	virtual bool IsOutputSrgb() const override;
 	//~ End IMediaTextureSample interface
 
-	bool ConfigureSample(uint32 InWidth, uint32 InHeight, uint32 InStride, ERivermaxMediaSourePixelFormat InSampleFormat, FTimespan InTime, const FFrameRate& InFrameRate, const TOptional<FTimecode>& InTimecode, bool bInIsSRGBInput);
+	bool ConfigureSample(uint32 InWidth, uint32 InHeight, uint32 InStride, ERivermaxMediaSourcePixelFormat InSampleFormat, FTimespan InTime, const FFrameRate& InFrameRate, const TOptional<FTimecode>& InTimecode, bool bInIsSRGBInput);
 
 private:
 
diff --git a/Engine/Plugins/VirtualProduction/Rivermax/RivermaxMedia/Source/RivermaxMedia/Private/RivermaxMediaTextureSampleConverter.cpp b/Engine/Plugins/VirtualProduction/Rivermax/RivermaxMedia/Source/RivermaxMedia/Private/RivermaxMediaTextureSampleConverter.cpp
--- a/Engine/Plugins/VirtualProduction/Rivermax/RivermaxMedia/Source/RivermaxMedia/Private/RivermaxMediaTextureSampleConverter.cpp
+++ b/Engine/Plugins/VirtualProduction/Rivermax/RivermaxMedia/Source/RivermaxMedia/Private/RivermaxMediaTextureSampleConverter.cpp
@@ -5,10 +5,10 @@
 #include "RenderGraphBuilder.h"
 #include "RivermaxShaders.h"
 #include "RivermaxMediaTextureSample.h"
+#include "RivermaxMediaUtils.h"
 
 
-
-void FRivermaxMediaTextureSampleConverter::Setup(ERivermaxMediaSourePixelFormat InPixelFormat, TWeakPtr<FRivermaxMediaTextureSample> InSample, bool bInDoSRGBToLinear)
+void FRivermaxMediaTextureSampleConverter::Setup(ERivermaxMediaSourcePixelFormat InPixelFormat, TWeakPtr<FRivermaxMediaTextureSample> InSample, bool bInDoSRGBToLinear)
 {
 	InputPixelFormat = InPixelFormat;
 	Sample = InSample;
@@ -17,7 +17,10 @@ void FRivermaxMediaTextureSampleConverter::Setup(ERivermaxMediaSourePixelFormat
 
 bool FRivermaxMediaTextureSampleConverter::Convert(FTexture2DRHIRef& InDestinationTexture, const FConversionHints& Hints)
 {
+	TRACE_CPUPROFILER_EVENT_SCOPE(RivermaxSampleConverter::Convert);
+
 	using namespace UE::RivermaxShaders;
+	using namespace UE::RivermaxCore;
 
 	TSharedPtr<FRivermaxMediaTextureSample> SamplePtr = Sample.Pin();
 	if (SamplePtr.IsValid() == false)
@@ -33,17 +36,15 @@ bool FRivermaxMediaTextureSampleConverter::Convert(FTexture2DRHIRef& InDestinati
 	//Configure shader and add conversion pass based on desired pixel format
 	switch (InputPixelFormat)
 	{
-	case ERivermaxMediaSourePixelFormat::YUV422_8bit:
+	case ERivermaxMediaSourcePixelFormat::YUV422_8bit:
 	{
 		break;
 	}
-	case ERivermaxMediaSourePixelFormat::YUV422_10bit:
+	case ERivermaxMediaSourcePixelFormat::YUV422_10bit:
 	{
-		constexpr uint32 PixelPerGroup = 2;
-		constexpr uint32 BytesPerGroup = 5;
-		const uint32 HorizontalByteCount = (InDestinationTexture->GetDesc().Extent.X / PixelPerGroup) * BytesPerGroup;
+		TRACE_CPUPROFILER_EVENT_SCOPE(RivermaxSampleConverter::YUV10ShaderSetup);
 		const int32 BytesPerElement = sizeof(FYUV10Bit422ToRGBACS::FYUV10Bit422LEBuffer);
-		const int32 ElementsPerRow = HorizontalByteCount / BytesPerElement;
+		const int32 ElementsPerRow = FMath::CeilToInt32(SamplePtr->GetStride() / (float)BytesPerElement);
 		const int32 ElementCount = ElementsPerRow * InDestinationTexture->GetDesc().Extent.Y;
 
 		FYUV10Bit422ToRGBACS::FPermutationDomain PermutationVector;
@@ -68,13 +69,11 @@ bool FRivermaxMediaTextureSampleConverter::Convert(FTexture2DRHIRef& InDestinati
 			, GroupCount);
 		break;
 	}
-	case ERivermaxMediaSourePixelFormat::RGB_8bit:
+	case ERivermaxMediaSourcePixelFormat::RGB_8bit:
 	{
-		constexpr uint32 PixelPerGroup = 1;
-		constexpr uint32 BytesPerGroup = (PixelPerGroup * 3 * 8) / 8;
-		const uint32 HorizontalByteCount = (InDestinationTexture->GetDesc().Extent.X / PixelPerGroup) * BytesPerGroup;
+		TRACE_CPUPROFILER_EVENT_SCOPE(RivermaxSampleConverter::RGB8ShaderSetup);
 		const int32 BytesPerElement = sizeof(FRGB8BitToRGBA8CS::FRGB8BitBuffer);
-		const int32 ElementsPerRow = HorizontalByteCount / BytesPerElement;
+		const int32 ElementsPerRow = FMath::CeilToInt32(SamplePtr->GetStride() / (float)BytesPerElement);
 		const int32 ElementCount = ElementsPerRow * InDestinationTexture->GetDesc().Extent.Y;
 
 		FRGB8BitToRGBA8CS::FPermutationDomain PermutationVector;
@@ -97,13 +96,12 @@ bool FRivermaxMediaTextureSampleConverter::Convert(FTexture2DRHIRef& InDestinati
 			, GroupCount);
 		break;
 	}
-	case ERivermaxMediaSourePixelFormat::RGB_10bit:
+	case ERivermaxMediaSourcePixelFormat::RGB_10bit:
 	{
-		constexpr uint32 PixelPerGroup = 4;
-		constexpr uint32 BytesPerGroup = (PixelPerGroup * 3 * 10) / 8;
-		const uint32 HorizontalByteCount = (InDestinationTexture->GetDesc().Extent.X / PixelPerGroup) * BytesPerGroup;
+		TRACE_CPUPROFILER_EVENT_SCOPE(RivermaxSampleConverter::RGB10ShaderSetup);
+
 		const int32 BytesPerElement = sizeof(FRGBToRGB10BitCS::FRGB10BitBuffer);
-		const int32 ElementsPerRow = HorizontalByteCount / BytesPerElement;
+		const int32 ElementsPerRow = FMath::CeilToInt32(SamplePtr->GetStride() / (float)BytesPerElement);
 		const int32 ElementCount = ElementsPerRow * InDestinationTexture->GetDesc().Extent.Y;
 		
 		FRGB10BitToRGBA10CS::FPermutationDomain PermutationVector;
diff --git a/Engine/Plugins/VirtualProduction/Rivermax/RivermaxMedia/Source/RivermaxMedia/Private/RivermaxMediaTextureSampleConverter.h b/Engine/Plugins/VirtualProduction/Rivermax/RivermaxMedia/Source/RivermaxMedia/Private/RivermaxMediaTextureSampleConverter.h
--- a/Engine/Plugins/VirtualProduction/Rivermax/RivermaxMedia/Source/RivermaxMedia/Private/RivermaxMediaTextureSampleConverter.h
+++ b/Engine/Plugins/VirtualProduction/Rivermax/RivermaxMedia/Source/RivermaxMedia/Private/RivermaxMediaTextureSampleConverter.h
@@ -17,7 +17,7 @@ class FRivermaxMediaTextureSampleConverter : public IMediaTextureSampleConverter
 public:
 
 	/** Configures settings to convert incoming sample */
-	void Setup(ERivermaxMediaSourePixelFormat InPixelFormat, TWeakPtr<FRivermaxMediaTextureSample> InSample, bool bInDoSRGBToLinear);
+	void Setup(ERivermaxMediaSourcePixelFormat InPixelFormat, TWeakPtr<FRivermaxMediaTextureSample> InSample, bool bInDoSRGBToLinear);
 
 	//~ Begin IMediaTextureSampleConverter interface
 	virtual bool Convert(FTexture2DRHIRef& InDstTexture, const FConversionHints& Hints) override;
@@ -27,7 +27,7 @@ public:
 private:
 
 	/** Incoming pixel format used to know how to handle the input buffer */
-	ERivermaxMediaSourePixelFormat InputPixelFormat = ERivermaxMediaSourePixelFormat::RGB_10bit;
+	ERivermaxMediaSourcePixelFormat InputPixelFormat = ERivermaxMediaSourcePixelFormat::RGB_10bit;
 
 	/** Sample holding the buffer to convert */
 	TWeakPtr<FRivermaxMediaTextureSample> Sample;
diff --git a/Engine/Plugins/VirtualProduction/Rivermax/RivermaxMedia/Source/RivermaxMedia/Private/RivermaxMediaUtils.cpp b/Engine/Plugins/VirtualProduction/Rivermax/RivermaxMedia/Source/RivermaxMedia/Private/RivermaxMediaUtils.cpp
new file mode 100644
--- /dev/null
+++ b/Engine/Plugins/VirtualProduction/Rivermax/RivermaxMedia/Source/RivermaxMedia/Private/RivermaxMediaUtils.cpp
@@ -0,0 +1,70 @@
+// Copyright Epic Games, Inc. All Rights Reserved.
+
+#include "RivermaxMediaUtils.h"
+
+namespace UE::RivermaxMediaUtils::Private
+{
+	UE::RivermaxCore::ESamplingType MediaOutputPixelFormatToRivermaxSamplingType(ERivermaxMediaOutputPixelFormat InPixelFormat)
+	{
+		using namespace UE::RivermaxCore;
+
+		switch (InPixelFormat)
+		{
+		case ERivermaxMediaOutputPixelFormat::PF_8BIT_YUV422:
+		{
+			return ESamplingType::YUV422_8bit;
+		}
+		case ERivermaxMediaOutputPixelFormat::PF_10BIT_YUV422:
+		{
+			return ESamplingType::YUV422_10bit;
+		}
+		case ERivermaxMediaOutputPixelFormat::PF_8BIT_RGB:
+		{
+			return ESamplingType::RGB_8bit;
+		}
+		case ERivermaxMediaOutputPixelFormat::PF_10BIT_RGB:
+		{
+			return ESamplingType::RGB_10bit;
+		}
+		case ERivermaxMediaOutputPixelFormat::PF_FLOAT16_RGB:
+		{
+			return ESamplingType::RGB_16bitFloat;
+		}
+		default:
+		{
+			checkNoEntry();
+			return ESamplingType::RGB_10bit;
+		}
+		}
+	}
+
+	UE::RivermaxCore::ESamplingType MediaSourcePixelFormatToRivermaxSamplingType(ERivermaxMediaSourcePixelFormat InPixelFormat)
+	{
+		using namespace UE::RivermaxCore;
+
+		switch (InPixelFormat)
+		{
+		case ERivermaxMediaSourcePixelFormat::YUV422_8bit:
+		{
+			return ESamplingType::YUV422_8bit;
+		}
+		case ERivermaxMediaSourcePixelFormat::YUV422_10bit:
+		{
+			return ESamplingType::YUV422_10bit;
+		}
+		case ERivermaxMediaSourcePixelFormat::RGB_8bit:
+		{
+			return ESamplingType::RGB_8bit;
+		}
+		case ERivermaxMediaSourcePixelFormat::RGB_10bit:
+		{
+			return ESamplingType::RGB_10bit;
+		}
+		default:
+		{
+			checkNoEntry();
+			return ESamplingType::RGB_10bit;
+		}
+		}
+	}
+}
diff --git a/Engine/Plugins/VirtualProduction/Rivermax/RivermaxMedia/Source/RivermaxMedia/Private/RivermaxMediaUtils.h b/Engine/Plugins/VirtualProduction/Rivermax/RivermaxMedia/Source/RivermaxMedia/Private/RivermaxMediaUtils.h
new file mode 100644
--- /dev/null
+++ b/Engine/Plugins/VirtualProduction/Rivermax/RivermaxMedia/Source/RivermaxMedia/Private/RivermaxMediaUtils.h
@@ -0,0 +1,17 @@
+// Copyright Epic Games, Inc. All Rights Reserved.
+
+#pragma once
+
+#include "CoreMinimal.h"
+
+#include "RivermaxFormats.h"
+#include "RivermaxMediaSource.h"
+#include "RivermaxMediaOutput.h"
+
+namespace UE::RivermaxMediaUtils::Private
+{
+
+UE::RivermaxCore::ESamplingType MediaOutputPixelFormatToRivermaxSamplingType(ERivermaxMediaOutputPixelFormat InPixelFormat);
+UE::RivermaxCore::ESamplingType MediaSourcePixelFormatToRivermaxSamplingType(ERivermaxMediaSourcePixelFormat InPixelFormat);
+
+}
\ No newline at end of file
diff --git a/Engine/Plugins/VirtualProduction/Rivermax/RivermaxMedia/Source/RivermaxMedia/Public/RivermaxMediaOutput.h b/Engine/Plugins/VirtualProduction/Rivermax/RivermaxMedia/Source/RivermaxMedia/Public/RivermaxMediaOutput.h
--- a/Engine/Plugins/VirtualProduction/Rivermax/RivermaxMedia/Source/RivermaxMedia/Public/RivermaxMediaOutput.h
+++ b/Engine/Plugins/VirtualProduction/Rivermax/RivermaxMedia/Source/RivermaxMedia/Public/RivermaxMediaOutput.h
@@ -13,8 +13,8 @@
 UENUM()
 enum class ERivermaxMediaOutputPixelFormat : uint8
 {
-	PF_8BIT_YUV UMETA(DisplayName = "8bit YUV"),
-	PF_10BIT_YCBCR_422 UMETA(DisplayName = "10bit YCbCr 4:2:2"),
+	PF_8BIT_YUV422 UMETA(DisplayName = "8bit YUV422"),
+	PF_10BIT_YUV422 UMETA(DisplayName = "10bit YUV422"),
 	PF_8BIT_RGB UMETA(DisplayName = "8bit RGB"),
 	PF_10BIT_RGB UMETA(DisplayName = "10bit RGB"),
 	PF_FLOAT16_RGB UMETA(DisplayName = "16bit Float RGB")
@@ -57,7 +57,7 @@ public:
 	FFrameRate FrameRate = {24,1};
 	
 	UPROPERTY(EditAnywhere, BlueprintReadWrite, Category = "Format")
-	ERivermaxMediaOutputPixelFormat PixelFormat = ERivermaxMediaOutputPixelFormat::PF_8BIT_YUV;
+	ERivermaxMediaOutputPixelFormat PixelFormat = ERivermaxMediaOutputPixelFormat::PF_10BIT_RGB;
 
 	UPROPERTY(EditAnywhere, BlueprintReadWrite, Category = "Format")
 	FString InterfaceAddress;
diff --git a/Engine/Plugins/VirtualProduction/Rivermax/RivermaxMedia/Source/RivermaxMedia/Public/RivermaxMediaSource.h b/Engine/Plugins/VirtualProduction/Rivermax/RivermaxMedia/Source/RivermaxMedia/Public/RivermaxMediaSource.h
--- a/Engine/Plugins/VirtualProduction/Rivermax/RivermaxMedia/Source/RivermaxMedia/Public/RivermaxMediaSource.h
+++ b/Engine/Plugins/VirtualProduction/Rivermax/RivermaxMedia/Source/RivermaxMedia/Public/RivermaxMediaSource.h
@@ -12,7 +12,7 @@
  * Native data format.
  */
 UENUM()
-enum class ERivermaxMediaSourePixelFormat : uint8
+enum class ERivermaxMediaSourcePixelFormat : uint8
 {
 	YUV422_8bit UMETA(DisplayName = "8bit YUV422"),
 	YUV422_10bit UMETA(DisplayName = "10bit YUV422"),
@@ -38,7 +38,7 @@ public:
 	FFrameRate FrameRate = {24,1};
 	
 	UPROPERTY(EditAnywhere, BlueprintReadWrite, Category = "Format")
-	ERivermaxMediaSourePixelFormat PixelFormat = ERivermaxMediaSourePixelFormat::YUV422_8bit;
+	ERivermaxMediaSourcePixelFormat PixelFormat = ERivermaxMediaSourcePixelFormat::RGB_10bit;
 
 	/** Network card interface to use to receive data */
 	UPROPERTY(EditAnywhere, BlueprintReadWrite, Category = "Format")
