- Fixing CIS about rivermax output stream

diff --git a/Engine/Plugins/VirtualProduction/Rivermax/RivermaxCore/Source/RivermaxCore/Private/RivermaxOutputStream.cpp b/Engine/Plugins/VirtualProduction/Rivermax/RivermaxCore/Source/RivermaxCore/Private/RivermaxOutputStream.cpp
--- a/Engine/Plugins/VirtualProduction/Rivermax/RivermaxCore/Source/RivermaxCore/Private/RivermaxOutputStream.cpp
+++ b/Engine/Plugins/VirtualProduction/Rivermax/RivermaxCore/Source/RivermaxCore/Private/RivermaxOutputStream.cpp
@@ -34,7 +34,6 @@ namespace UE::RivermaxCore::Private
 		TEXT("Interval at which to show stats in seconds"),
 		ECVF_Default);
 
-
 	bool FindPayloadSize(const FRivermaxStreamOptions& InOptions, uint32 InBytesPerLine, const FVideoFormatInfo& FormatInfo, uint16& OutPayloadSize)
 	{
 		using namespace UE::RivermaxCore::Private::Utils;
@@ -56,13 +55,12 @@ namespace UE::RivermaxCore::Private
 			return false;
 		}
 
-		bool bContinue = true;
-		while (bContinue)
+		while (true)
 		{
 			const int32 TestSize = InBytesPerLine / TestPoint;
 			if (TestSize < MinPayloadSize)
 			{
-				return false;
+				break;
 			}
 
 			if (TestSize <= MaxPayloadSize)
