- Adding missing pixel format support to rmax 2110 input and output (yuv8, rgb12, rgb16f)

diff --git a/Engine/Plugins/VirtualProduction/Rivermax/RivermaxCore/Shaders/Private/RivermaxShaders.usf b/Engine/Plugins/VirtualProduction/Rivermax/RivermaxCore/Shaders/Private/RivermaxShaders.usf
--- a/Engine/Plugins/VirtualProduction/Rivermax/RivermaxCore/Shaders/Private/RivermaxShaders.usf
+++ b/Engine/Plugins/VirtualProduction/Rivermax/RivermaxCore/Shaders/Private/RivermaxShaders.usf
@@ -21,6 +21,13 @@
 uint InputTexturePixelsPerThread;
 RWTexture2D<float4> OutTexture;
 
+struct FYUV8Bit422Buffer
+{
+	uint DWord0;
+};
+RWStructuredBuffer<FYUV8Bit422Buffer>	OutYUV4228bitBuffer;
+StructuredBuffer<FYUV8Bit422Buffer>	InputYUV4228bitBuffer;
+
 struct FYCbCr10Bit422
 {
 	uint DWord0;
@@ -48,12 +55,48 @@
 RWStructuredBuffer<FRGB10BitBuffer>	OutRGB10Buffer;
 StructuredBuffer<FRGB10BitBuffer> InputBuffer;
 
+struct FRGB12BitBuffer
+{
+	uint DWords[9];
+};
+RWStructuredBuffer<FRGB12BitBuffer>	OutRGB12Buffer;
+StructuredBuffer<FRGB12BitBuffer> InputRGB12Buffer;
+
 struct FRGB16fBuffer
 {
-	half Pixel0[3];
-	half Pixel1[3];
+	uint DWords[3];
 };
 RWStructuredBuffer<FRGB16fBuffer> OutRGB16fBuffer;
+StructuredBuffer<FRGB16fBuffer> InputRGB16fBuffer;
 
 
+
+// Shader to convert from RGB to YUV 422 8 bits
+
+[numthreads(8,8,1)]
+void RGBToYUV8Bit422(uint3 ThreadId : SV_DispatchThreadID, uint GroupIndex : SV_GroupIndex)
+{
+	// Make sure it's a valid thread
+	if(ThreadId.x >= HorizontalElementCount || ThreadId.y >= VerticalElementCount)
+	{
+		return;
+	}
+
+	uint Index = ThreadId.y * HorizontalElementCount + ThreadId.x;
+
+	float X = (ThreadId.x * InputTexturePixelsPerThread) + InputPixelOffsetX;
+	float Y = (ThreadId.y + InputPixelOffsetY) * OnePixelDeltaY;
+
+	// We process 2 RGB pixels per thread
+	float3 RGB0 = InputTexture.SampleLevel(InputSampler, float2(((X + 0.0) + 0.5) * OnePixelDeltaX, Y), 0).rgb;
+	float3 RGB1 = InputTexture.SampleLevel(InputSampler, float2(((X + 1.0) + 0.5) * OnePixelDeltaX, Y), 0).rgb;
+
+	// The 8 RGB needs to be converted to YUV. Each encoded pixel will use 2 RGB pixels (YnUnVnYn+1)
+	uint3 YUV0 = RgbToYuv(RGB0, ColorTransform, DoLinearToSrgb) * 255;
+	uint3 YUV1 = RgbToYuv(RGB1, ColorTransform, DoLinearToSrgb) * 255;
+
+	OutYUV4228bitBuffer[Index].DWord0 = (YUV1.x << 24) | (YUV0.z << 16) | (YUV0.x << 8) | (YUV0.y << 0);
+}
+
+
 // Shader to convert from RGB to YUV 422 10 bits packed little endian
@@ -100,35 +143,6 @@
 	OutYCbCrBuffer[Index].DWord4 = ((YUV7.x & 0xFF) << 24) | ((YUV6.z & 0x3F) << 18) | ((YUV7.x >> 8) << 16) | ((YUV6.x & 0xF) << 12) | ((YUV6.z >> 6) << 8) | ((YUV6.y & 0x3) << 6) | (YUV6.x >> 4);
 }
 
-
-// Shader to convert from RGBA to RGB8
-
-[numthreads(8,8,1)]
-void RGBToRGB8BitCS(uint3 ThreadId : SV_DispatchThreadID, uint GroupIndex : SV_GroupIndex)
-{
-	// Make sure it's a valid thread
-	if(ThreadId.x >= HorizontalElementCount || ThreadId.y >= VerticalElementCount)
-	{
-		return;
-	}
-
-	uint Index = ThreadId.y * HorizontalElementCount + ThreadId.x;
-
-	float X = (ThreadId.x * InputTexturePixelsPerThread) + InputPixelOffsetX;
-	float Y = (ThreadId.y + InputPixelOffsetY) * OnePixelDeltaY;
-
-	// We process 4 RGB pixels per thread
-	uint3 RGB0 = InputTexture.SampleLevel(InputSampler, float2(((X + 0.0) + 0.5) * OnePixelDeltaX, Y), 0).rgb * 255;
-	uint3 RGB1 = InputTexture.SampleLevel(InputSampler, float2(((X + 1.0) + 0.5) * OnePixelDeltaX, Y), 0).rgb * 255;
-	uint3 RGB2 = InputTexture.SampleLevel(InputSampler, float2(((X + 2.0) + 0.5) * OnePixelDeltaX, Y), 0).rgb * 255;
-	uint3 RGB3 = InputTexture.SampleLevel(InputSampler, float2(((X + 3.0) + 0.5) * OnePixelDeltaX, Y), 0).rgb * 255;
-
-	// We write the 4 RGB pixels (24 bits * 4 = 96 bits = 3x32 bits) following the packing 2110 expects
-	OutRGB8Buffer[Index].DWord0 = (RGB1.x << 24) | (RGB0.z << 16) | (RGB0.y << 8) | (RGB0.x << 0);
-	OutRGB8Buffer[Index].DWord1 = (RGB2.y << 24) | (RGB2.x << 16) | (RGB1.z << 8) | (RGB1.y << 0);
-	OutRGB8Buffer[Index].DWord2 = (RGB3.z << 24) | (RGB3.y << 16) | (RGB3.x << 8) | (RGB2.z << 0);
-}
-
 // Shader to convert from RGBA to RGB10
 
 [numthreads(8,8,1)]
@@ -145,22 +159,42 @@
 	float X = (ThreadId.x * InputTexturePixelsPerThread) + InputPixelOffsetX;
 	float Y = (ThreadId.y + InputPixelOffsetY) * OnePixelDeltaY;
 
-	uint3 RGB00 = InputTexture.SampleLevel(InputSampler, float2(((X + 0.0) + 0.5) * OnePixelDeltaX, Y), 0).rgb * 1023;
-	uint3 RGB01 = InputTexture.SampleLevel(InputSampler, float2(((X + 1.0) + 0.5) * OnePixelDeltaX, Y), 0).rgb * 1023;
-	uint3 RGB02 = InputTexture.SampleLevel(InputSampler, float2(((X + 2.0) + 0.5) * OnePixelDeltaX, Y), 0).rgb * 1023;
-	uint3 RGB03 = InputTexture.SampleLevel(InputSampler, float2(((X + 3.0) + 0.5) * OnePixelDeltaX, Y), 0).rgb * 1023;
-	uint3 RGB04 = InputTexture.SampleLevel(InputSampler, float2(((X + 4.0) + 0.5) * OnePixelDeltaX, Y), 0).rgb * 1023;
-	uint3 RGB05 = InputTexture.SampleLevel(InputSampler, float2(((X + 5.0) + 0.5) * OnePixelDeltaX, Y), 0).rgb * 1023;
-	uint3 RGB06 = InputTexture.SampleLevel(InputSampler, float2(((X + 6.0) + 0.5) * OnePixelDeltaX, Y), 0).rgb * 1023;
-	uint3 RGB07 = InputTexture.SampleLevel(InputSampler, float2(((X + 7.0) + 0.5) * OnePixelDeltaX, Y), 0).rgb * 1023;
-	uint3 RGB08 = InputTexture.SampleLevel(InputSampler, float2(((X + 8.0) + 0.5) * OnePixelDeltaX, Y), 0).rgb * 1023;
-	uint3 RGB09 = InputTexture.SampleLevel(InputSampler, float2(((X + 9.0) + 0.5) * OnePixelDeltaX, Y), 0).rgb * 1023;
-	uint3 RGB10 = InputTexture.SampleLevel(InputSampler, float2(((X + 10.0) + 0.5) * OnePixelDeltaX, Y), 0).rgb * 1023;
-	uint3 RGB11 = InputTexture.SampleLevel(InputSampler, float2(((X + 11.0) + 0.5) * OnePixelDeltaX, Y), 0).rgb * 1023;
-	uint3 RGB12 = InputTexture.SampleLevel(InputSampler, float2(((X + 12.0) + 0.5) * OnePixelDeltaX, Y), 0).rgb * 1023;
-	uint3 RGB13 = InputTexture.SampleLevel(InputSampler, float2(((X + 13.0) + 0.5) * OnePixelDeltaX, Y), 0).rgb * 1023;
-	uint3 RGB14 = InputTexture.SampleLevel(InputSampler, float2(((X + 14.0) + 0.5) * OnePixelDeltaX, Y), 0).rgb * 1023;
-	uint3 RGB15 = InputTexture.SampleLevel(InputSampler, float2(((X + 15.0) + 0.5) * OnePixelDeltaX, Y), 0).rgb * 1023;
+	uint Denormalizer = 1023;
+	uint3 RGB00 = InputTexture.SampleLevel(InputSampler, float2(((X + 0.0) + 0.5) * OnePixelDeltaX, Y), 0).rgb * Denormalizer;
+	uint3 RGB01 = InputTexture.SampleLevel(InputSampler, float2(((X + 1.0) + 0.5) * OnePixelDeltaX, Y), 0).rgb * Denormalizer;
+	uint3 RGB02 = InputTexture.SampleLevel(InputSampler, float2(((X + 2.0) + 0.5) * OnePixelDeltaX, Y), 0).rgb * Denormalizer;
+	uint3 RGB03 = InputTexture.SampleLevel(InputSampler, float2(((X + 3.0) + 0.5) * OnePixelDeltaX, Y), 0).rgb * Denormalizer;
+	uint3 RGB04 = InputTexture.SampleLevel(InputSampler, float2(((X + 4.0) + 0.5) * OnePixelDeltaX, Y), 0).rgb * Denormalizer;
+	uint3 RGB05 = InputTexture.SampleLevel(InputSampler, float2(((X + 5.0) + 0.5) * OnePixelDeltaX, Y), 0).rgb * Denormalizer;
+	uint3 RGB06 = InputTexture.SampleLevel(InputSampler, float2(((X + 6.0) + 0.5) * OnePixelDeltaX, Y), 0).rgb * Denormalizer;
+	uint3 RGB07 = InputTexture.SampleLevel(InputSampler, float2(((X + 7.0) + 0.5) * OnePixelDeltaX, Y), 0).rgb * Denormalizer;
+	uint3 RGB08 = InputTexture.SampleLevel(InputSampler, float2(((X + 8.0) + 0.5) * OnePixelDeltaX, Y), 0).rgb * Denormalizer;
+	uint3 RGB09 = InputTexture.SampleLevel(InputSampler, float2(((X + 9.0) + 0.5) * OnePixelDeltaX, Y), 0).rgb * Denormalizer;
+	uint3 RGB10 = InputTexture.SampleLevel(InputSampler, float2(((X + 10.0) + 0.5) * OnePixelDeltaX, Y), 0).rgb * Denormalizer;
+	uint3 RGB11 = InputTexture.SampleLevel(InputSampler, float2(((X + 11.0) + 0.5) * OnePixelDeltaX, Y), 0).rgb * Denormalizer;
+	uint3 RGB12 = InputTexture.SampleLevel(InputSampler, float2(((X + 12.0) + 0.5) * OnePixelDeltaX, Y), 0).rgb * Denormalizer;
+	uint3 RGB13 = InputTexture.SampleLevel(InputSampler, float2(((X + 13.0) + 0.5) * OnePixelDeltaX, Y), 0).rgb * Denormalizer;
+	uint3 RGB14 = InputTexture.SampleLevel(InputSampler, float2(((X + 14.0) + 0.5) * OnePixelDeltaX, Y), 0).rgb * Denormalizer;
+	uint3 RGB15 = InputTexture.SampleLevel(InputSampler, float2(((X + 15.0) + 0.5) * OnePixelDeltaX, Y), 0).rgb * Denormalizer;
 
+	// Cap values to be in range
+	RGB00 = min(RGB00, Denormalizer);
+	RGB01 = min(RGB01, Denormalizer);
+	RGB02 = min(RGB02, Denormalizer);
+	RGB03 = min(RGB03, Denormalizer);
+	RGB04 = min(RGB04, Denormalizer);
+	RGB05 = min(RGB05, Denormalizer);
+	RGB06 = min(RGB06, Denormalizer);
+	RGB07 = min(RGB07, Denormalizer);
+	RGB08 = min(RGB08, Denormalizer);
+	RGB09 = min(RGB09, Denormalizer);
+	RGB10 = min(RGB10, Denormalizer);
+	RGB11 = min(RGB11, Denormalizer);
+	RGB12 = min(RGB12, Denormalizer);
+	RGB13 = min(RGB13, Denormalizer);
+	RGB14 = min(RGB14, Denormalizer);
+	RGB15 = min(RGB15, Denormalizer);
+
+
 	// All this bit swizzling is to pack 3 channels of 10 bits with R being in LSB of the first 32 bits and subsequent channels following until we used all 16 pixels.
 	// Example : RGB = (0x123, 0x245, 0x0AB)
@@ -191,91 +225,6 @@
 	OutRGB10Buffer[Index].DWords[14] =							   ((RGB15.z & 0xFF) << 24)  | ((RGB15.y & 0x3F) << 18) | ((RGB15.z >> 8)   << 16) | ((RGB15.x & 0xF)  << 12) | ((RGB15.y >> 6)   << 8)  | ((RGB14.z & 0x3)  << 6)  | (RGB15.x >> 4);
 }
 
-// Shader to convert from RGBA to RGB Float16
-
-[numthreads(8,8,1)]
-void RGBToRGB16fBitCS(uint3 ThreadId : SV_DispatchThreadID, uint GroupIndex : SV_GroupIndex)
-{
-	// Make sure it's a valid thread
-	if(ThreadId.x >= HorizontalElementCount || ThreadId.y >= VerticalElementCount)
-	{
-		return;
-	}
-
-	float X = (ThreadId.x * InputTexturePixelsPerThread) + InputPixelOffsetX;
-	float Y = (ThreadId.y + InputPixelOffsetY) * OnePixelDeltaY;
-
-	half3 RGB0 = InputTexture.SampleLevel(InputSampler, float2(((X + 0.0) + 0.5) * OnePixelDeltaX, Y), 0).rgb;
-	half3 RGB1 = InputTexture.SampleLevel(InputSampler, float2(((X + 1.0) + 0.5) * OnePixelDeltaX, Y), 0).rgb;
-	
-	// Storage has been untested as no available supported that viewer. Will need to be validated. 
-	uint Index = ThreadId.y * HorizontalElementCount + ThreadId.x;
-	OutRGB16fBuffer[Index].Pixel0[0]  = RGB0.z;
-	OutRGB16fBuffer[Index].Pixel0[1]  = RGB0.y;
-	OutRGB16fBuffer[Index].Pixel0[2]  = RGB0.x;
-	OutRGB16fBuffer[Index].Pixel1[0]  = RGB1.z;
-	OutRGB16fBuffer[Index].Pixel1[1]  = RGB1.y;
-	OutRGB16fBuffer[Index].Pixel1[2]  = RGB1.x;
-}
-
-// Shader to convert from packed RGB8 to RGBA8 texture
-
-[numthreads(8,8,1)]
-void RGB8BitToRGBA8CS(uint3 ThreadId : SV_DispatchThreadID, uint GroupIndex : SV_GroupIndex)
-{
-	// Make sure it's a valid thread
-	if(ThreadId.x >= HorizontalElementCount || ThreadId.y >= VerticalElementCount)
-	{
-		return;
-	}
-
-	// Pixels we treat in shader is pretty much hardcoded to 16 so no parameter has been used
-	const uint OutputTexturePixelsPerThread = 4;
-
-	// Scale U (x index) based on the number of pixels we treat per thread
-	const uint BaseU = ThreadId.x * OutputTexturePixelsPerThread;
-	const uint BaseV = ThreadId.y;
-
-	// Input buffer indexing is done in one dimension so adjust using thread Id and elements per row
-	const uint BufferIndex = ThreadId.y * HorizontalElementCount + ThreadId.x;
-
-	// Used to normalize pixel values based on 8 bits input
-	const float Mul = 1.0 / 255.0;
-
-	float3 RGB00;
-	RGB00.x = (((InputRGB8Buffer[BufferIndex].DWord0 >> 00) & 0xFF)) * Mul;
-	RGB00.y = (((InputRGB8Buffer[BufferIndex].DWord0 >>  8) & 0xFF)) * Mul;
-	RGB00.z = (((InputRGB8Buffer[BufferIndex].DWord0 >> 16) & 0xFF)) * Mul;
-
-	float3 RGB01;
-	RGB01.x = (((InputRGB8Buffer[BufferIndex].DWord0 >> 24) & 0xFF)) * Mul;
-	RGB01.y = (((InputRGB8Buffer[BufferIndex].DWord1 >> 00) & 0xFF)) * Mul;
-	RGB01.z = (((InputRGB8Buffer[BufferIndex].DWord1 >>  8) & 0xFF)) * Mul;
-
-	float3 RGB02;
-	RGB02.x = (((InputRGB8Buffer[BufferIndex].DWord1 >> 16) & 0xFF)) * Mul;
-	RGB02.y = (((InputRGB8Buffer[BufferIndex].DWord1 >> 24) & 0xFF)) * Mul;
-	RGB02.z = (((InputRGB8Buffer[BufferIndex].DWord2 >> 00) & 0xFF)) * Mul;
-	
-	float3 RGB03;
-	RGB03.x = (((InputRGB8Buffer[BufferIndex].DWord2 >>  8) & 0xFF)) * Mul;
-	RGB03.y = (((InputRGB8Buffer[BufferIndex].DWord2 >> 16) & 0xFF)) * Mul;
-	RGB03.z = (((InputRGB8Buffer[BufferIndex].DWord2 >> 24) & 0xFF)) * Mul;
-	
-
-	#if DO_SRGB_TO_LINEAR
-		RGB00 = sRGBToLinear(RGB00);
-		RGB01 = sRGBToLinear(RGB01);
-		RGB02 = sRGBToLinear(RGB02);
-		RGB03 = sRGBToLinear(RGB03);
-	#endif
-
-	OutTexture[uint2(BaseU + 0, BaseV)] = float4(RGB00, 1.0);
-	OutTexture[uint2(BaseU + 1, BaseV)] = float4(RGB01, 1.0);
-	OutTexture[uint2(BaseU + 2, BaseV)] = float4(RGB02, 1.0);
-	OutTexture[uint2(BaseU + 3, BaseV)] = float4(RGB03, 1.0);
-}
-
 // Shader to convert from packed RGB10 to RGB10A2 texture
 
 [numthreads(8,8,1)]
@@ -418,6 +367,390 @@
 	OutTexture[uint2(BaseU + 15, BaseV)] = float4(RGB15, 1.0);
 }
 
+// Shader to convert from RGBA to RGB8
 
+[numthreads(8,8,1)]
+void RGBToRGB8BitCS(uint3 ThreadId : SV_DispatchThreadID, uint GroupIndex : SV_GroupIndex)
+{
+	// Make sure it's a valid thread
+	if(ThreadId.x >= HorizontalElementCount || ThreadId.y >= VerticalElementCount)
+	{
+		return;
+	}
+
+	uint Index = ThreadId.y * HorizontalElementCount + ThreadId.x;
+
+	float X = (ThreadId.x * InputTexturePixelsPerThread) + InputPixelOffsetX;
+	float Y = (ThreadId.y + InputPixelOffsetY) * OnePixelDeltaY;
+
+	// We process 4 RGB pixels per thread
+	uint Denormalizer = 255;
+	uint3 RGB0 = InputTexture.SampleLevel(InputSampler, float2(((X + 0.0) + 0.5) * OnePixelDeltaX, Y), 0).rgb * Denormalizer;
+	uint3 RGB1 = InputTexture.SampleLevel(InputSampler, float2(((X + 1.0) + 0.5) * OnePixelDeltaX, Y), 0).rgb * Denormalizer;
+	uint3 RGB2 = InputTexture.SampleLevel(InputSampler, float2(((X + 2.0) + 0.5) * OnePixelDeltaX, Y), 0).rgb * Denormalizer;
+	uint3 RGB3 = InputTexture.SampleLevel(InputSampler, float2(((X + 3.0) + 0.5) * OnePixelDeltaX, Y), 0).rgb * Denormalizer;
+
+	// Cap values to be in range
+	RGB0 = min(RGB0, Denormalizer);
+	RGB1 = min(RGB1, Denormalizer);
+	RGB2 = min(RGB2, Denormalizer);
+	RGB3 = min(RGB3, Denormalizer);
+
+	// We write the 4 RGB pixels (24 bits * 4 = 96 bits = 3x32 bits) following the packing 2110 expects
+	OutRGB8Buffer[Index].DWord0 = (RGB1.x << 24) | (RGB0.z << 16) | (RGB0.y << 8) | (RGB0.x << 0);
+	OutRGB8Buffer[Index].DWord1 = (RGB2.y << 24) | (RGB2.x << 16) | (RGB1.z << 8) | (RGB1.y << 0);
+	OutRGB8Buffer[Index].DWord2 = (RGB3.z << 24) | (RGB3.y << 16) | (RGB3.x << 8) | (RGB2.z << 0);
+}
+
+// Shader to convert from packed RGB8 to RGBA8 texture
+
+[numthreads(8,8,1)]
+void RGB8BitToRGBA8CS(uint3 ThreadId : SV_DispatchThreadID, uint GroupIndex : SV_GroupIndex)
+{
+	// Make sure it's a valid thread
+	if(ThreadId.x >= HorizontalElementCount || ThreadId.y >= VerticalElementCount)
+	{
+		return;
+	}
+
+	// Pixels we treat in shader is pretty much hardcoded to 16 so no parameter has been used
+	const uint OutputTexturePixelsPerThread = 4;
+
+	// Scale U (x index) based on the number of pixels we treat per thread
+	const uint BaseU = ThreadId.x * OutputTexturePixelsPerThread;
+	const uint BaseV = ThreadId.y;
+
+	// Input buffer indexing is done in one dimension so adjust using thread Id and elements per row
+	const uint BufferIndex = ThreadId.y * HorizontalElementCount + ThreadId.x;
+
+	// Used to normalize pixel values based on 8 bits input
+	const float Mul = 1.0 / 255.0;
+
+	float3 RGB00;
+	RGB00.x = (((InputRGB8Buffer[BufferIndex].DWord0 >> 00) & 0xFF)) * Mul;
+	RGB00.y = (((InputRGB8Buffer[BufferIndex].DWord0 >>  8) & 0xFF)) * Mul;
+	RGB00.z = (((InputRGB8Buffer[BufferIndex].DWord0 >> 16) & 0xFF)) * Mul;
+
+	float3 RGB01;
+	RGB01.x = (((InputRGB8Buffer[BufferIndex].DWord0 >> 24) & 0xFF)) * Mul;
+	RGB01.y = (((InputRGB8Buffer[BufferIndex].DWord1 >> 00) & 0xFF)) * Mul;
+	RGB01.z = (((InputRGB8Buffer[BufferIndex].DWord1 >>  8) & 0xFF)) * Mul;
+
+	float3 RGB02;
+	RGB02.x = (((InputRGB8Buffer[BufferIndex].DWord1 >> 16) & 0xFF)) * Mul;
+	RGB02.y = (((InputRGB8Buffer[BufferIndex].DWord1 >> 24) & 0xFF)) * Mul;
+	RGB02.z = (((InputRGB8Buffer[BufferIndex].DWord2 >> 00) & 0xFF)) * Mul;
+	
+	float3 RGB03;
+	RGB03.x = (((InputRGB8Buffer[BufferIndex].DWord2 >>  8) & 0xFF)) * Mul;
+	RGB03.y = (((InputRGB8Buffer[BufferIndex].DWord2 >> 16) & 0xFF)) * Mul;
+	RGB03.z = (((InputRGB8Buffer[BufferIndex].DWord2 >> 24) & 0xFF)) * Mul;
+	
+
+	#if DO_SRGB_TO_LINEAR
+		RGB00 = sRGBToLinear(RGB00);
+		RGB01 = sRGBToLinear(RGB01);
+		RGB02 = sRGBToLinear(RGB02);
+		RGB03 = sRGBToLinear(RGB03);
+	#endif
+
+	OutTexture[uint2(BaseU + 0, BaseV)] = float4(RGB00, 1.0);
+	OutTexture[uint2(BaseU + 1, BaseV)] = float4(RGB01, 1.0);
+	OutTexture[uint2(BaseU + 2, BaseV)] = float4(RGB02, 1.0);
+	OutTexture[uint2(BaseU + 3, BaseV)] = float4(RGB03, 1.0);
+}
+
+// Shader to convert from RGBA to RGB12
+
+[numthreads(8,8,1)]
+void RGBToRGB12BitCS(uint3 ThreadId : SV_DispatchThreadID, uint GroupIndex : SV_GroupIndex)
+{
+	// Make sure it's a valid thread
+	if(ThreadId.x >= HorizontalElementCount || ThreadId.y >= VerticalElementCount)
+	{
+		return;
+	}
+
+	uint Index = ThreadId.y * HorizontalElementCount + ThreadId.x;
+
+	float X = (ThreadId.x * InputTexturePixelsPerThread) + InputPixelOffsetX;
+	float Y = (ThreadId.y + InputPixelOffsetY) * OnePixelDeltaY;
+
+	//12 bits max value
+	uint Denormalizer = 4095;
+	uint3 RGB00 = InputTexture.SampleLevel(InputSampler, float2(((X + 0.0) + 0.5) * OnePixelDeltaX, Y), 0).rgb * Denormalizer;
+	uint3 RGB01 = InputTexture.SampleLevel(InputSampler, float2(((X + 1.0) + 0.5) * OnePixelDeltaX, Y), 0).rgb * Denormalizer;
+	uint3 RGB02 = InputTexture.SampleLevel(InputSampler, float2(((X + 2.0) + 0.5) * OnePixelDeltaX, Y), 0).rgb * Denormalizer;
+	uint3 RGB03 = InputTexture.SampleLevel(InputSampler, float2(((X + 3.0) + 0.5) * OnePixelDeltaX, Y), 0).rgb * Denormalizer;
+	uint3 RGB04 = InputTexture.SampleLevel(InputSampler, float2(((X + 4.0) + 0.5) * OnePixelDeltaX, Y), 0).rgb * Denormalizer;
+	uint3 RGB05 = InputTexture.SampleLevel(InputSampler, float2(((X + 5.0) + 0.5) * OnePixelDeltaX, Y), 0).rgb * Denormalizer;
+	uint3 RGB06 = InputTexture.SampleLevel(InputSampler, float2(((X + 6.0) + 0.5) * OnePixelDeltaX, Y), 0).rgb * Denormalizer;
+	uint3 RGB07 = InputTexture.SampleLevel(InputSampler, float2(((X + 7.0) + 0.5) * OnePixelDeltaX, Y), 0).rgb * Denormalizer;
+
+	// Cap values to be in range
+	RGB00 = min(RGB00, Denormalizer);
+	RGB01 = min(RGB01, Denormalizer);
+	RGB02 = min(RGB02, Denormalizer);
+	RGB03 = min(RGB03, Denormalizer);
+	RGB04 = min(RGB04, Denormalizer);
+	RGB05 = min(RGB05, Denormalizer);
+	RGB06 = min(RGB06, Denormalizer);
+	RGB07 = min(RGB07, Denormalizer);
+
+	OutRGB12Buffer[Index].DWords[0] =							  ((RGB00.z >> 4)   << 24) |							 ((RGB00.y & 0xFF)  << 16) |((RGB00.x & 0xF)  << 12) | ((RGB00.y >> 8)    << 8)								| (RGB00.x >> 4);
+	OutRGB12Buffer[Index].DWords[1] = ((RGB01.y & 0x00F) << 28) | ((RGB01.z >> 8)   << 24) |							 ((RGB01.y >> 4)    << 16)							 | ((RGB01.x & 0xFF)  << 8) | ((RGB00.z & 0x00F) << 4)	| (RGB01.x >> 8);
+	OutRGB12Buffer[Index].DWords[2] =							  ((RGB02.y & 0xFF) << 24) | ((RGB02.x & 0xF) << 20)	| ((RGB02.y >> 8)   << 16)							  | ((RGB02.x >> 4)	  << 8)								| (RGB01.z & 0xFF);
+	OutRGB12Buffer[Index].DWords[3] =							  ((RGB03.y >> 4)   << 24)								| ((RGB03.x & 0xFF) << 16) | ((RGB02.z & 0xF) << 12) | ((RGB03.x >> 8)   << 8)								| (RGB02.z >> 4);
+	OutRGB12Buffer[Index].DWords[4] = ((RGB04.x & 0x00F) << 28) | ((RGB04.y >> 8)   << 24)								| ((RGB04.x >> 4)   << 16)							  | ((RGB03.z & 0xFF) << 8)  | ((RGB03.y & 0x00F)  << 4)| (RGB03.z >> 8);
+	OutRGB12Buffer[Index].DWords[5] =							  ((RGB05.x & 0xFF) << 24) | ((RGB04.z & 0xF) << 20)	| ((RGB05.x >> 8)   << 16)							  | ((RGB04.z >> 4)   << 8)								| (RGB04.y & 0xFF);
+	OutRGB12Buffer[Index].DWords[6] =							  ((RGB06.x >> 4)   << 24)								| ((RGB05.z & 0xFF) << 16) | ((RGB05.y & 0xF) << 12)  | ((RGB05.z >> 8)   << 8)								| (RGB05.y >> 4);
+	OutRGB12Buffer[Index].DWords[7] = ((RGB06.z & 0x00F) << 28) | ((RGB07.x >> 8)   << 24)								| ((RGB06.z >> 4)   << 16)							  | ((RGB06.y & 0xFF) << 8)  | ((RGB06.x & 0x00F) << 4) | (RGB06.y >> 8);
+	OutRGB12Buffer[Index].DWords[8] =							  ((RGB07.z & 0xFF) << 24) | ((RGB07.y & 0xF) << 20)	| ((RGB07.z >> 8)   << 16)							  | ((RGB07.y >> 4)   << 8)								| (RGB07.x & 0xFF);
+}
+
+// Shader to convert from packed RGB12 to float RGBA texture
+
+[numthreads(8,8,1)]
+void RGB12BitToRGBACS(uint3 ThreadId : SV_DispatchThreadID, uint GroupIndex : SV_GroupIndex)
+{
+	// Make sure it's a valid thread
+	if(ThreadId.x >= HorizontalElementCount || ThreadId.y >= VerticalElementCount)
+	{
+		return;
+	}
+
+	// Pixels we treat in shader is pretty much hardcoded to 16 so no parameter has been used
+	const uint OutputTexturePixelsPerThread = 8;
+
+	// Scale U (x index) based on the number of pixels we treat per thread
+	const uint BaseU = ThreadId.x * OutputTexturePixelsPerThread;
+	const uint BaseV = ThreadId.y;
+
+	// Input buffer indexing is done in one dimension so adjust using thread Id and elements per row
+	const uint BufferIndex = ThreadId.y * HorizontalElementCount + ThreadId.x;
+
+	// Used to normalize pixel values based on 12 bits input
+	const float Mul = 1.0 / 4095.0;
+
+	float3 RGB00 = uint3(0,0,0);
+	RGB00.x = ((((InputRGB12Buffer[BufferIndex].DWords[00] >>  0)  << 4) | ((InputRGB12Buffer[BufferIndex].DWords[00] >> 12) & 0x00F)) & 0xFFF) * Mul;
+	RGB00.y = ((((InputRGB12Buffer[BufferIndex].DWords[00] >>  8)  << 8) | ((InputRGB12Buffer[BufferIndex].DWords[00] >> 16) & 0x0FF)) & 0xFFF) * Mul;
+	RGB00.z = ((((InputRGB12Buffer[BufferIndex].DWords[00] >> 24)  << 4) | ((InputRGB12Buffer[BufferIndex].DWords[01] >> 04) & 0x00F)) & 0xFFF) * Mul;
+
+	float3 RGB01 = uint3(0,0,0);
+	RGB01.x = ((((InputRGB12Buffer[BufferIndex].DWords[01] >>  0) << 8) | ((InputRGB12Buffer[BufferIndex].DWords[01] >>  8) & 0x0FF)) & 0xFFF) * Mul;
+	RGB01.y = ((((InputRGB12Buffer[BufferIndex].DWords[01] >> 16) << 4) | ((InputRGB12Buffer[BufferIndex].DWords[01] >> 28) & 0x00F)) & 0xFFF) * Mul;
+	RGB01.z = ((((InputRGB12Buffer[BufferIndex].DWords[01] >> 24) << 8) | ((InputRGB12Buffer[BufferIndex].DWords[02] >> 00) & 0x0FF)) & 0xFFF) * Mul;
+
+	float3 RGB02 = uint3(0,0,0);
+	RGB02.x = ((((InputRGB12Buffer[BufferIndex].DWords[02] >>  8) << 4) | ((InputRGB12Buffer[BufferIndex].DWords[02] >> 20) & 0x00F)) & 0xFFF) * Mul;
+	RGB02.y = ((((InputRGB12Buffer[BufferIndex].DWords[02] >> 16) << 8) | ((InputRGB12Buffer[BufferIndex].DWords[02] >> 24) & 0x0FF)) & 0xFFF) * Mul;
+	RGB02.z = ((((InputRGB12Buffer[BufferIndex].DWords[03] >>  0) << 4) | ((InputRGB12Buffer[BufferIndex].DWords[03] >> 12) & 0x00F)) & 0xFFF) * Mul;
+	
+	float3 RGB03 = uint3(0,0,0);
+	RGB03.x = ((((InputRGB12Buffer[BufferIndex].DWords[03] >>  8) << 8) | ((InputRGB12Buffer[BufferIndex].DWords[03] >> 16) & 0x0FF)) & 0xFFF) * Mul;
+	RGB03.y = ((((InputRGB12Buffer[BufferIndex].DWords[03] >> 24) << 4) | ((InputRGB12Buffer[BufferIndex].DWords[04] >> 04) & 0x00F)) & 0xFFF) * Mul;
+	RGB03.z = ((((InputRGB12Buffer[BufferIndex].DWords[04] >>  0) << 8) | ((InputRGB12Buffer[BufferIndex].DWords[04] >>  8) & 0x0FF)) & 0xFFF) * Mul;
+	
+	float3 RGB04 = uint3(0,0,0);
+	RGB04.x = ((((InputRGB12Buffer[BufferIndex].DWords[04] >> 16) << 4) | ((InputRGB12Buffer[BufferIndex].DWords[04] >> 28) & 0x00F)) & 0xFFF) * Mul;
+	RGB04.y = ((((InputRGB12Buffer[BufferIndex].DWords[04] >> 24) << 8) | ((InputRGB12Buffer[BufferIndex].DWords[05] >>  0) & 0x0FF)) & 0xFFF) * Mul;
+	RGB04.z = ((((InputRGB12Buffer[BufferIndex].DWords[05] >>  8) << 4) | ((InputRGB12Buffer[BufferIndex].DWords[05] >> 20) & 0x00F)) & 0xFFF) * Mul;
+	
+	float3 RGB05;
+	RGB05.x = ((((InputRGB12Buffer[BufferIndex].DWords[05] >> 16) << 8) | ((InputRGB12Buffer[BufferIndex].DWords[05] >> 24) & 0x0FF)) & 0xFFF) * Mul;
+	RGB05.y = ((((InputRGB12Buffer[BufferIndex].DWords[06] >> 00) << 4) | ((InputRGB12Buffer[BufferIndex].DWords[06] >> 12) & 0x00F)) & 0xFFF) * Mul;
+	RGB05.z = ((((InputRGB12Buffer[BufferIndex].DWords[06] >>  8) << 8) | ((InputRGB12Buffer[BufferIndex].DWords[06] >> 16) & 0x0FF)) & 0xFFF) * Mul;
+
+	float3 RGB06 = uint3(0,0,0);
+	RGB06.x = ((((InputRGB12Buffer[BufferIndex].DWords[06] >> 24) << 4) | ((InputRGB12Buffer[BufferIndex].DWords[07] >> 04) & 0x00F)) & 0xFFF) * Mul;
+	RGB06.y = ((((InputRGB12Buffer[BufferIndex].DWords[07] >>  0) << 8) | ((InputRGB12Buffer[BufferIndex].DWords[07] >>  8) & 0x0FF)) & 0xFFF) * Mul;
+	RGB06.z = ((((InputRGB12Buffer[BufferIndex].DWords[07] >>  16) << 4) | ((InputRGB12Buffer[BufferIndex].DWords[07] >> 28) & 0x00F)) & 0xFFF) * Mul;
+
+	float3 RGB07 = uint3(0,0,0);
+	RGB07.x = ((((InputRGB12Buffer[BufferIndex].DWords[07] >> 24) << 8) | ((InputRGB12Buffer[BufferIndex].DWords[8] >> 00) & 0x0FF)) & 0xFFF) * Mul;
+	RGB07.y = ((((InputRGB12Buffer[BufferIndex].DWords[8] >>  8) << 4) | ((InputRGB12Buffer[BufferIndex].DWords[8] >> 20) & 0x00F)) & 0xFFF) * Mul;
+	RGB07.z = ((((InputRGB12Buffer[BufferIndex].DWords[8] >>  16) << 8) | ((InputRGB12Buffer[BufferIndex].DWords[8] >> 24) & 0x0FF)) & 0xFFF) * Mul;
+
+	#if DO_SRGB_TO_LINEAR
+		RGB00 = sRGBToLinear(RGB00);
+		RGB01 = sRGBToLinear(RGB01);
+		RGB02 = sRGBToLinear(RGB02);
+		RGB03 = sRGBToLinear(RGB03);
+		RGB04 = sRGBToLinear(RGB04);
+		RGB05 = sRGBToLinear(RGB05);
+		RGB06 = sRGBToLinear(RGB06);
+		RGB07 = sRGBToLinear(RGB07);
+	#endif
+
+	OutTexture[uint2(BaseU + 0, BaseV)] = float4(RGB00, 1.0);
+	OutTexture[uint2(BaseU + 1, BaseV)] = float4(RGB01, 1.0);
+	OutTexture[uint2(BaseU + 2, BaseV)] = float4(RGB02, 1.0);
+	OutTexture[uint2(BaseU + 3, BaseV)] = float4(RGB03, 1.0);
+	OutTexture[uint2(BaseU + 4, BaseV)] = float4(RGB04, 1.0);
+	OutTexture[uint2(BaseU + 5, BaseV)] = float4(RGB05, 1.0);
+	OutTexture[uint2(BaseU + 6, BaseV)] = float4(RGB06, 1.0);
+	OutTexture[uint2(BaseU + 7, BaseV)] = float4(RGB07, 1.0);
+}
+
+// Shader to convert from RGBA to RGB Float16
+
+
+
+[numthreads(8,8,1)]
+void RGBToRGB16fBitCS(uint3 ThreadId : SV_DispatchThreadID, uint GroupIndex : SV_GroupIndex)
+{
+	// Make sure it's a valid thread
+	if(ThreadId.x >= HorizontalElementCount || ThreadId.y >= VerticalElementCount)
+	{
+		return;
+	}
+
+	float X = (ThreadId.x * InputTexturePixelsPerThread) + InputPixelOffsetX;
+	float Y = (ThreadId.y + InputPixelOffsetY) * OnePixelDeltaY;
+
+	float3 RGB0 = InputTexture.SampleLevel(InputSampler, float2(((X + 0.0) + 0.5) * OnePixelDeltaX, Y), 0).rgb;
+	float3 RGB1 = InputTexture.SampleLevel(InputSampler, float2(((X + 1.0) + 0.5) * OnePixelDeltaX, Y), 0).rgb;
+	//RGB0 = float3(1,0,1);
+	//RGB1 = float3(0,1,0);
+
+	// Storage has been untested as no available supported that viewer. Will need to be validated. 
+	uint Index = ThreadId.y * HorizontalElementCount + ThreadId.x;
+
+	// Convert float to half representation and write it in lowest 16 bits of a uint
+	uint IntR0 = f32tof16(RGB0.x);
+	uint IntG0 = f32tof16(RGB0.y);
+	uint IntB0 = f32tof16(RGB0.z);
+	uint IntR1 = f32tof16(RGB1.x);
+	uint IntG1 = f32tof16(RGB1.y);
+	uint IntB1 = f32tof16(RGB1.z);
+
+	#if 1
+	OutRGB16fBuffer[Index].DWords[0]  = ((IntG0 & 0xFF) << 24) | ((IntG0 >> 8) << 16) | ((IntR0 & 0xFF) << 8) | ((IntR0 >> 8) << 0);
+	OutRGB16fBuffer[Index].DWords[1]  = ((IntR1 & 0xFF) << 24) | ((IntR1 >> 8) << 16) | ((IntB0 & 0xFF) << 8) | ((IntB0 >> 8) << 0);
+	OutRGB16fBuffer[Index].DWords[2]  = ((IntB1 & 0xFF) << 24) | ((IntB1 >> 8) << 16) | ((IntG1 & 0xFF) << 8) | ((IntG1 >> 8) << 0);
+	#else
+	OutRGB16fBuffer[Index].DWords[0]  = IntR0;
+	OutRGB16fBuffer[Index].DWords[1]  = IntG0;
+	OutRGB16fBuffer[Index].DWords[2]  = IntB0;
+	#endif
+}
+
+/**
+ * This function extracts the 16bit (float) data stored in uint (32 bit)
+ * Same used in EXRSwizzler. 
+ */
+float UintContainingFloatToFloat(in uint Value)
+{
+	int Sign = (Value >> 15) & 0x00000001;
+	int Exponent = (Value >> 10) & 0x0000001f;
+	int Mantissa = Value & 0x000003ff;
+
+	Exponent = Exponent + (127 - 15);
+	Mantissa = Mantissa << 13;
+
+	return asfloat((Sign << 31) | (Exponent << 23) | Mantissa);
+}
+
+float GetEncodedFloatValue(in uint Value)
+{
+	uint LowerWord = Value & 0xFFFF;
+	uint Swizzled = (LowerWord << 8);
+	Swizzled += (LowerWord >> 8) & 0xFF;
+
+	#if SM6_PROFILE
+	return UintContainingFloatToFloat(Swizzled);
+	#else
+	return f16tof32(Swizzled);
+	#endif
+}
+
+// Shader to convert from packed RGB16f to float RGBA texture
+
+[numthreads(8,8,1)]
+void RGB16fBitToRGBACS(uint3 ThreadId : SV_DispatchThreadID, uint GroupIndex : SV_GroupIndex)
+{
+	// Make sure it's a valid thread
+	if(ThreadId.x >= HorizontalElementCount || ThreadId.y >= VerticalElementCount)
+	{
+		return;
+	}
+
+	// Pixels we treat in shader is pretty much hardcoded to 2 so no parameter has been used
+	const uint OutputTexturePixelsPerThread = 2;
+
+	// Scale U (x index) based on the number of pixels we treat per thread
+	const uint BaseU = ThreadId.x * OutputTexturePixelsPerThread;
+	const uint BaseV = ThreadId.y;
+
+	// Input buffer indexing is done in one dimension so adjust using thread Id and elements per row
+	const uint BufferIndex = ThreadId.y * HorizontalElementCount + ThreadId.x;
+
+	const uint Dword0 = InputRGB16fBuffer[BufferIndex].DWords[00];
+	const uint Dword1 = InputRGB16fBuffer[BufferIndex].DWords[01];
+	const uint Dword2 = InputRGB16fBuffer[BufferIndex].DWords[02];
+
+	const uint IntRed0 = Dword0 & 0xFFFF;
+	const uint IntGreen0 = (Dword0 >> 16) & 0xFFFF;
+	const uint IntBlue0 = Dword1 & 0xFFFF;
+	const uint IntRed1 = (Dword1 >> 16) & 0xFFFF;
+	const uint IntGreen1 = Dword2 & 0xFFFF;
+	const uint IntBlue1 = (Dword2 >> 16) & 0xFFFF;
+
+	float4 RGB00 = float4(0,0,0,1);
+	float4 RGB01 = float4(0,0,0,1);
+	RGB00.x = GetEncodedFloatValue(IntRed0);
+	RGB00.y = GetEncodedFloatValue(IntGreen0);
+	RGB00.z = GetEncodedFloatValue(IntBlue0);
+	RGB01.x = GetEncodedFloatValue(IntRed1);
+	RGB01.y = GetEncodedFloatValue(IntGreen1);
+	RGB01.z = GetEncodedFloatValue(IntBlue1);
+	
+	OutTexture[uint2(BaseU + 0, BaseV)] = RGB00;
+	OutTexture[uint2(BaseU + 1, BaseV)] = RGB01;
+}
+
+// Shader to convert from YUV 422 8 bits to RGBA
+
+[numthreads(8,8,1)]
+void YUV8Bit422ToRGBACS(uint3 ThreadId : SV_DispatchThreadID, uint GroupIndex : SV_GroupIndex)
+{
+	// Make sure it's a valid thread
+	if(ThreadId.x >= HorizontalElementCount || ThreadId.y >= VerticalElementCount)
+	{
+		return;
+	}
+
+	// Pixels we treat in shader is pretty much hardcoded to 2 so no parameter has been used
+	const uint OutputTexturePixelsPerThread = 2;
+
+	// Scale U (x index) based on the number of pixels we treat per thread
+	const uint BaseU = ThreadId.x * OutputTexturePixelsPerThread;
+	const uint BaseV = ThreadId.y;
+
+	// Input buffer indexing is done in one dimension so adjust using thread Id and elements per row
+	const uint BufferIndex = ThreadId.y * HorizontalElementCount + ThreadId.x;
+
+	// Used to normalize pixel values based on 8 bits input
+	const float Mul = 1.0 / 255.0;
+	
+	float U0 = ((InputYUV4228bitBuffer[BufferIndex].DWord0 >>  0) & 0xFF) * Mul;
+	float Y0 = ((InputYUV4228bitBuffer[BufferIndex].DWord0 >>  8) & 0xFF) * Mul;
+	float V0 = ((InputYUV4228bitBuffer[BufferIndex].DWord0 >> 16) & 0xFF) * Mul;
+	float Y1 = ((InputYUV4228bitBuffer[BufferIndex].DWord0 >> 24) & 0xFF) * Mul;
+
+	#if DO_SRGB_TO_LINEAR
+	int DoSRGBToLinear = 1;
+	#else
+	int DoSRGBToLinear = 0;
+	#endif
+
+	float3 RGB00 = YuvToRgb(float3(Y0, U0, V0), ColorTransform, DoSRGBToLinear);
+	float3 RGB01 = YuvToRgb(float3(Y1, U0, V0), ColorTransform, DoSRGBToLinear);
+
+	OutTexture[uint2(BaseU + 0, BaseV)] = float4(RGB00, 1.0);
+	OutTexture[uint2(BaseU + 1, BaseV)] = float4(RGB01, 1.0);
+}
+
 // Shader to convert from YUV 422 10 bits packed little endian to RGBA
 
diff --git a/Engine/Plugins/VirtualProduction/Rivermax/RivermaxCore/Source/RivermaxRendering/Private/RivermaxShaders.cpp b/Engine/Plugins/VirtualProduction/Rivermax/RivermaxCore/Source/RivermaxRendering/Private/RivermaxShaders.cpp
--- a/Engine/Plugins/VirtualProduction/Rivermax/RivermaxCore/Source/RivermaxRendering/Private/RivermaxShaders.cpp
+++ b/Engine/Plugins/VirtualProduction/Rivermax/RivermaxCore/Source/RivermaxRendering/Private/RivermaxShaders.cpp
@@ -8,6 +8,58 @@
 
 namespace UE::RivermaxShaders
 {
+	/* FYUV8Bit422ToRGBACS shader
+	*****************************************************************************/
+
+	IMPLEMENT_GLOBAL_SHADER(FYUV8Bit422ToRGBACS, "/Plugin/RivermaxCore/Private/RivermaxShaders.usf", "YUV8Bit422ToRGBACS", SF_Compute);
+
+	FYUV8Bit422ToRGBACS::FParameters* FYUV8Bit422ToRGBACS::AllocateAndSetParameters(FRDGBuilder& GraphBuilder, FRDGBufferRef YUVBuffer, FRDGTextureRef OutputTexture, const FMatrix& ColorTransform, const FVector& YUVOffset, int32 BufferElementsPerRow, int32 BufferLineCount)
+	{
+		FYUV8Bit422ToRGBACS::FParameters* Parameters = GraphBuilder.AllocParameters<FYUV8Bit422ToRGBACS::FParameters>();
+
+		Parameters->InputYUV4228bitBuffer= GraphBuilder.CreateSRV(YUVBuffer);
+		Parameters->ColorTransform = (FMatrix44f)MediaShaders::CombineColorTransformAndOffset(ColorTransform, YUVOffset);
+		Parameters->HorizontalElementCount = BufferElementsPerRow;
+		Parameters->VerticalElementCount = BufferLineCount;
+		Parameters->OutTexture = GraphBuilder.CreateUAV(OutputTexture, ERDGUnorderedAccessViewFlags::None);
+
+		return Parameters;
+	}
+
+	/* FRGBToYUV8Bit422CS shader
+	 *****************************************************************************/
+
+	IMPLEMENT_GLOBAL_SHADER(FRGBToYUV8Bit422CS, "/Plugin/RivermaxCore/Private/RivermaxShaders.usf", "RGBToYUV8Bit422", SF_Compute);
+
+	FRGBToYUV8Bit422CS::FParameters* FRGBToYUV8Bit422CS::AllocateAndSetParameters(FRDGBuilder& GraphBuilder, FRDGTextureRef RGBATexture, const FIntPoint& SourceSize, const FIntRect& SourceViewRect, const FIntPoint& OutputSize, const FMatrix& ColorTransform, const FVector& YUVOffset, bool bDoLinearToSrgb, FRDGBufferRef OutputBuffer)
+	{
+		FRGBToYUV8Bit422CS::FParameters* Parameters = GraphBuilder.AllocParameters<FRGBToYUV8Bit422CS::FParameters>();
+
+		const uint32 InputTextureSizeX = SourceSize.X;
+		const uint32 InputTextureSizeY = SourceSize.Y;
+		Parameters->RGBToYUVConversion.InputTexture = RGBATexture;
+		Parameters->RGBToYUVConversion.InputSampler = TStaticSamplerState<SF_Point>::GetRHI();
+		Parameters->RGBToYUVConversion.ColorTransform = (FMatrix44f)MediaShaders::CombineColorTransformAndOffset(ColorTransform, YUVOffset);
+		Parameters->RGBToYUVConversion.DoLinearToSrgb = bDoLinearToSrgb;
+		Parameters->RGBToYUVConversion.OnePixelDeltaX = 1.0f / (float)InputTextureSizeX;
+		Parameters->OnePixelDeltaY = 1.0f / (float)InputTextureSizeY;
+
+		// Used to offset indices to sample input texture based on ThreadId of CS
+		Parameters->InputPixelOffsetX = SourceViewRect.Min.X;
+		Parameters->InputPixelOffsetY = SourceViewRect.Min.Y;
+
+		// Output size will be based on RGB to YUV422 encoding ( divide by 2 ) 
+		Parameters->HorizontalElementCount = OutputSize.X;
+		Parameters->VerticalElementCount = OutputSize.Y;
+
+		// Each thread will read many RGB pixels to convert to YUV
+		Parameters->InputTexturePixelsPerThread = InputTextureSizeX / OutputSize.X;
+
+		Parameters->OutYUV4228bitBuffer = GraphBuilder.CreateUAV(OutputBuffer);
+
+		return Parameters;
+	}
+
 	/* FYUV10Bit422ToRGBACS shader
 	 *****************************************************************************/
 
@@ -158,6 +210,56 @@ namespace UE::RivermaxShaders
 		return Parameters;
 	}
 
+	/* FRGB12BitToRGBA12CS shader
+	*****************************************************************************/
+
+	IMPLEMENT_GLOBAL_SHADER(FRGB12BitToRGBA12CS, "/Plugin/RivermaxCore/Private/RivermaxShaders.usf", "RGB12BitToRGBACS", SF_Compute);
+
+
+	FRGB12BitToRGBA12CS::FParameters* FRGB12BitToRGBA12CS::AllocateAndSetParameters(FRDGBuilder& GraphBuilder, FRDGBufferRef RGBBuffer, FRDGTextureRef OutputTexture, int32 BufferElementsPerRow, int32 BufferLineCount)
+	{
+		FRGB12BitToRGBA12CS::FParameters* Parameters = GraphBuilder.AllocParameters<FRGB12BitToRGBA12CS::FParameters>();
+
+		Parameters->InputRGB12Buffer = GraphBuilder.CreateSRV(RGBBuffer);
+		Parameters->HorizontalElementCount = BufferElementsPerRow;
+		Parameters->VerticalElementCount = BufferLineCount;
+		Parameters->OutTexture = GraphBuilder.CreateUAV(OutputTexture, ERDGUnorderedAccessViewFlags::None);
+
+		return Parameters;
+	}
+
+	/* FRGBToRGB12BitCS shader
+	 *****************************************************************************/
+
+	IMPLEMENT_GLOBAL_SHADER(FRGBToRGB12BitCS, "/Plugin/RivermaxCore/Private/RivermaxShaders.usf", "RGBToRGB12BitCS", SF_Compute);
+
+
+	FRGBToRGB12BitCS::FParameters* FRGBToRGB12BitCS::AllocateAndSetParameters(FRDGBuilder& GraphBuilder, FRDGTextureRef RGBATexture, const FIntPoint& SourceSize, const FIntRect& SourceViewRect, const FIntPoint& OutputSize, FRDGBufferRef OutputBuffer)
+	{
+		FRGBToRGB12BitCS::FParameters* Parameters = GraphBuilder.AllocParameters<FRGBToRGB12BitCS::FParameters>();
+
+		const uint32 InputTextureSizeX = SourceSize.X;
+		const uint32 InputTextureSizeY = SourceSize.Y;
+		Parameters->InputTexture = RGBATexture;
+		Parameters->InputSampler = TStaticSamplerState<SF_Point>::GetRHI();
+		Parameters->OnePixelDeltaX = 1.0f / (float)InputTextureSizeX;
+		Parameters->OnePixelDeltaY = 1.0f / (float)InputTextureSizeY;
+
+		// Used to offset indices to sample input texture based on ThreadId of CS
+		Parameters->InputPixelOffsetX = SourceViewRect.Min.X;
+		Parameters->InputPixelOffsetY = SourceViewRect.Min.Y;
+
+		Parameters->HorizontalElementCount = OutputSize.X;
+		Parameters->VerticalElementCount = OutputSize.Y;
+
+		// Each thread will read many RGBA pixels to pack into RGB buffer
+		Parameters->InputTexturePixelsPerThread = InputTextureSizeX / OutputSize.X;
+
+		Parameters->OutRGB12Buffer = GraphBuilder.CreateUAV(OutputBuffer);
+
+		return Parameters;
+	}
+
 	/* FRGBToRGB16fCS shader
 	 *****************************************************************************/
 
@@ -189,5 +291,23 @@ namespace UE::RivermaxShaders
 
 		return Parameters;
 	}
+
+	/* FRGB16fBitToRGBA16fCS shader
+	*****************************************************************************/
+
+	IMPLEMENT_GLOBAL_SHADER(FRGB16fBitToRGBA16fCS, "/Plugin/RivermaxCore/Private/RivermaxShaders.usf", "RGB16fBitToRGBACS", SF_Compute);
+
+
+	FRGB16fBitToRGBA16fCS::FParameters* FRGB16fBitToRGBA16fCS::AllocateAndSetParameters(FRDGBuilder& GraphBuilder, FRDGBufferRef RGBBuffer, FRDGTextureRef OutputTexture, int32 BufferElementsPerRow, int32 BufferLineCount)
+	{
+		FRGB16fBitToRGBA16fCS::FParameters* Parameters = GraphBuilder.AllocParameters<FRGB16fBitToRGBA16fCS::FParameters>();
+
+		Parameters->InputRGB16fBuffer = GraphBuilder.CreateSRV(RGBBuffer);
+		Parameters->HorizontalElementCount = BufferElementsPerRow;
+		Parameters->VerticalElementCount = BufferLineCount;
+		Parameters->OutTexture = GraphBuilder.CreateUAV(OutputTexture, ERDGUnorderedAccessViewFlags::None);
+
+		return Parameters;
+	}
 }
 
diff --git a/Engine/Plugins/VirtualProduction/Rivermax/RivermaxCore/Source/RivermaxRendering/Public/RivermaxShaders.h b/Engine/Plugins/VirtualProduction/Rivermax/RivermaxCore/Source/RivermaxRendering/Public/RivermaxShaders.h
--- a/Engine/Plugins/VirtualProduction/Rivermax/RivermaxCore/Source/RivermaxRendering/Public/RivermaxShaders.h
+++ b/Engine/Plugins/VirtualProduction/Rivermax/RivermaxCore/Source/RivermaxRendering/Public/RivermaxShaders.h
@@ -13,6 +13,84 @@
 
 namespace UE::RivermaxShaders
 {
+	/**
+	 * Compute shader to convert 2110 YUV8 to RGBA
+	 */
+	class RIVERMAXRENDERING_API FYUV8Bit422ToRGBACS : public FGlobalShader
+	{
+	public:
+
+		//Structure definition to match StructuredBuffer in RivermaxShaders.usf
+		struct FYUV8Bit422Buffer
+		{
+			uint32 DWord0;
+		};
+
+		DECLARE_GLOBAL_SHADER(FYUV8Bit422ToRGBACS);
+		SHADER_USE_PARAMETER_STRUCT(FYUV8Bit422ToRGBACS, FGlobalShader);
+
+		class FSRGBToLinear : SHADER_PERMUTATION_BOOL("DO_SRGB_TO_LINEAR");
+		using FPermutationDomain = TShaderPermutationDomain<FSRGBToLinear>;
+
+		BEGIN_SHADER_PARAMETER_STRUCT(FParameters, )
+			SHADER_PARAMETER_RDG_BUFFER_SRV(StructuredBuffer<FYUV8Bit422Buffer>, InputYUV4228bitBuffer)
+			SHADER_PARAMETER(FMatrix44f, ColorTransform)
+			SHADER_PARAMETER(uint32, HorizontalElementCount)
+			SHADER_PARAMETER(uint32, VerticalElementCount)
+			SHADER_PARAMETER_RDG_TEXTURE_UAV(RWTexture2D, OutTexture)
+			END_SHADER_PARAMETER_STRUCT()
+
+	public:
+
+		// Called by the engine to determine which permutations to compile for this shader
+		static bool ShouldCompilePermutation(const FGlobalShaderPermutationParameters& Parameters)
+		{
+			return IsFeatureLevelSupported(Parameters.Platform, ERHIFeatureLevel::SM5);
+		}
+
+		/** Allocates and setup shader parameter in the incoming graph builder */
+		FYUV8Bit422ToRGBACS::FParameters* AllocateAndSetParameters(FRDGBuilder& GraphBuilder, FRDGBufferRef YUVBuffer, FRDGTextureRef OutputTexture, const FMatrix& ColorTransform, const FVector& YUVOffset, int32 BufferElementsPerRow, int32 BufferLineCount);
+	};
+
+	/**
+	 * Compute shader to convert RGB to YUV422 8 bits
+	 */
+	class RIVERMAXRENDERING_API FRGBToYUV8Bit422CS : public FGlobalShader
+	{
+	public:
+		//Structure definition to match StructuredBuffer in RivermaxShaders.usf
+		struct FYUV8Bit422Buffer
+		{
+			uint32 DWord0;
+		};
+
+		DECLARE_GLOBAL_SHADER(FRGBToYUV8Bit422CS);
+
+		SHADER_USE_PARAMETER_STRUCT(FRGBToYUV8Bit422CS, FGlobalShader);
+
+		BEGIN_SHADER_PARAMETER_STRUCT(FParameters, )
+			SHADER_PARAMETER_STRUCT_INCLUDE(FRGBToYUVConversion, RGBToYUVConversion)
+			SHADER_PARAMETER(float, OnePixelDeltaY)
+			SHADER_PARAMETER(float, InputPixelOffsetX)
+			SHADER_PARAMETER(float, InputPixelOffsetY)
+			SHADER_PARAMETER(uint32, HorizontalElementCount)
+			SHADER_PARAMETER(uint32, VerticalElementCount)
+			SHADER_PARAMETER(uint32, InputTexturePixelsPerThread)
+			SHADER_PARAMETER_RDG_BUFFER_UAV(RWStructuredBuffer<FYUV8Bit422Buffer>, OutYUV4228bitBuffer)
+			END_SHADER_PARAMETER_STRUCT()
+
+	public:
+
+		// Called by the engine to determine which permutations to compile for this shader
+		static bool ShouldCompilePermutation(const FGlobalShaderPermutationParameters& Parameters)
+		{
+			return IsFeatureLevelSupported(Parameters.Platform, ERHIFeatureLevel::SM5);
+		}
+
+		/** Allocates and setup shader parameter in the incoming graph builder */
+		FRGBToYUV8Bit422CS::FParameters* AllocateAndSetParameters(FRDGBuilder& GraphBuilder, FRDGTextureRef RGBATexture, const FIntPoint& SourceSize, const FIntRect& SourceViewRect, const FIntPoint& OutputSize, const FMatrix& ColorTransform, const FVector& YUVOffset, bool bDoLinearToSrgb, FRDGBufferRef OutputBuffer);
+	};
+
 	/**
 	 * Compute shader to convert packed 8 bits RGB to RGBA 8bits
 	 */
@@ -248,7 +326,7 @@ namespace UE::RivermaxShaders
 
 		DECLARE_GLOBAL_SHADER(FRGB10BitToRGBA10CS);
 		SHADER_USE_PARAMETER_STRUCT(FRGB10BitToRGBA10CS, FGlobalShader);
-		
+
 		class FSRGBToLinear : SHADER_PERMUTATION_BOOL("DO_SRGB_TO_LINEAR");
 		using FPermutationDomain = TShaderPermutationDomain<FSRGBToLinear>;
 
@@ -257,7 +335,7 @@ namespace UE::RivermaxShaders
 			SHADER_PARAMETER(uint32, HorizontalElementCount)
 			SHADER_PARAMETER(uint32, VerticalElementCount)
 			SHADER_PARAMETER_RDG_TEXTURE_UAV(RWTexture2D, OutTexture)
-		END_SHADER_PARAMETER_STRUCT()
+			END_SHADER_PARAMETER_STRUCT()
 
 	public:
 
@@ -271,6 +349,90 @@ namespace UE::RivermaxShaders
 		FRGB10BitToRGBA10CS::FParameters* AllocateAndSetParameters(FRDGBuilder& GraphBuilder, FRDGBufferRef RGBBuffer, FRDGTextureRef OutputTexture, int32 BufferElementsPerRow, int32 BufferLineCount);
 	};
 
+	/**
+	 * Compute shader to convert RGBA to packed 12 bits RGB
+	 */
+	class RIVERMAXRENDERING_API FRGBToRGB12BitCS : public FGlobalShader
+	{
+	public:
+
+		// Structure definition to match StructuredBuffer in RivermaxShaders.usf
+		// 12 bits per channel, 36 bits per pixel. Aligns on 9 32bits -> 288bits -> 8x RGB12 pixels
+		// 2110-20 says it's 2 pixels per pgroup. 
+		struct FRGB12BitBuffer
+		{
+			uint32 DWords[9];
+		};
+
+		DECLARE_GLOBAL_SHADER(FRGBToRGB12BitCS);
+
+		SHADER_USE_PARAMETER_STRUCT(FRGBToRGB12BitCS, FGlobalShader);
+
+		BEGIN_SHADER_PARAMETER_STRUCT(FParameters, )
+			SHADER_PARAMETER_RDG_TEXTURE(Texture2D, InputTexture)
+			SHADER_PARAMETER_SAMPLER(SamplerState, InputSampler)
+			SHADER_PARAMETER(float, OnePixelDeltaX)
+			SHADER_PARAMETER(float, OnePixelDeltaY)
+			SHADER_PARAMETER(float, InputPixelOffsetX)
+			SHADER_PARAMETER(float, InputPixelOffsetY)
+			SHADER_PARAMETER(uint32, HorizontalElementCount)
+			SHADER_PARAMETER(uint32, VerticalElementCount)
+			SHADER_PARAMETER(uint32, InputTexturePixelsPerThread)
+			SHADER_PARAMETER_RDG_BUFFER_UAV(RWStructuredBuffer<FRGB12BitBuffer>, OutRGB12Buffer)
+			END_SHADER_PARAMETER_STRUCT()
+
+	public:
+
+		// Called by the engine to determine which permutations to compile for this shader
+		static bool ShouldCompilePermutation(const FGlobalShaderPermutationParameters& Parameters)
+		{
+			return IsFeatureLevelSupported(Parameters.Platform, ERHIFeatureLevel::SM5);
+		}
+
+		/** Allocates and setup shader parameter in the incoming graph builder */
+		FRGBToRGB12BitCS::FParameters* AllocateAndSetParameters(FRDGBuilder& GraphBuilder, FRDGTextureRef RGBATexture, const FIntPoint& SourceSize, const FIntRect& SourceViewRect, const FIntPoint& OutputSize, FRDGBufferRef OutputBuffer);
+	};
+
+	/**
+	 * Compute shader to convert packed 12 bits RGB to RGBA
+	 */
+	class RIVERMAXRENDERING_API FRGB12BitToRGBA12CS : public FGlobalShader
+	{
+	public:
+
+		// Structure definition to match StructuredBuffer in RivermaxShaders.usf
+		// 12 bits per channel, 36 bits per pixel. Aligns on 9 32bits -> 288bits -> 8x RGB12 pixels
+		// 2110-20 says it's 2 pixels per pgroup. 
+		struct FRGB12BitBuffer
+		{
+			uint32 DWords[9];
+		};
+
+		DECLARE_GLOBAL_SHADER(FRGB12BitToRGBA12CS);
+		SHADER_USE_PARAMETER_STRUCT(FRGB12BitToRGBA12CS, FGlobalShader);
+
+		class FSRGBToLinear : SHADER_PERMUTATION_BOOL("DO_SRGB_TO_LINEAR");
+		using FPermutationDomain = TShaderPermutationDomain<FSRGBToLinear>;
+
+		BEGIN_SHADER_PARAMETER_STRUCT(FParameters, )
+			SHADER_PARAMETER_RDG_BUFFER_SRV(StructuredBuffer<FRGB12BitBuffer>, InputRGB12Buffer)
+			SHADER_PARAMETER(uint32, HorizontalElementCount)
+			SHADER_PARAMETER(uint32, VerticalElementCount)
+			SHADER_PARAMETER_RDG_TEXTURE_UAV(RWTexture2D, OutTexture)
+			END_SHADER_PARAMETER_STRUCT()
+
+	public:
+
+		// Called by the engine to determine which permutations to compile for this shader
+		static bool ShouldCompilePermutation(const FGlobalShaderPermutationParameters& Parameters)
+		{
+			return IsFeatureLevelSupported(Parameters.Platform, ERHIFeatureLevel::SM5);
+		}
+
+		/** Allocates and setup shader parameter in the incoming graph builder */
+		FRGB12BitToRGBA12CS::FParameters* AllocateAndSetParameters(FRDGBuilder& GraphBuilder, FRDGBufferRef RGBBuffer, FRDGTextureRef OutputTexture, int32 BufferElementsPerRow, int32 BufferLineCount);
+	};
+
 	/**
 	 * Compute shader to convert RGBA to Float16 RGB
 	 */
@@ -286,8 +448,7 @@ namespace UE::RivermaxShaders
 		// To align on 32bits, use 2 pixels (16 bit * 3 * 2 = 3x32bits)
 		struct FRGB16fBuffer
 		{
-			uint16 Pixel0[3];
-			uint16 Pixel1[3];
+			uint32 DWords[3];
 		};
 
 		BEGIN_SHADER_PARAMETER_STRUCT(FParameters, )
@@ -314,5 +475,44 @@ namespace UE::RivermaxShaders
 		/** Allocates and setup shader parameter in the incoming graph builder */
 		FRGBToRGB16fCS::FParameters* AllocateAndSetParameters(FRDGBuilder& GraphBuilder, FRDGTextureRef RGBATexture, const FIntPoint& SourceSize, const FIntRect& SourceViewRect, const FIntPoint& OutputSize, FRDGBufferRef OutputBuffer);
 	};
+
+	/**
+	 * Compute shader to convert 16 bits float RGB to RGBA
+	 */
+	class RIVERMAXRENDERING_API FRGB16fBitToRGBA16fCS : public FGlobalShader
+	{
+	public:
+
+		// Structure definition to match StructuredBuffer in RivermaxShaders.usf
+		// To align on 32bits, use 2 pixels (16 bit * 3 * 2 = 3x32bits)
+		struct FRGB16fBuffer
+		{
+			uint32 DWords[3];
+		};
+
+		DECLARE_GLOBAL_SHADER(FRGB16fBitToRGBA16fCS);
+		SHADER_USE_PARAMETER_STRUCT(FRGB16fBitToRGBA16fCS, FGlobalShader);
+
+		class FSRGBToLinear : SHADER_PERMUTATION_BOOL("DO_SRGB_TO_LINEAR");
+		using FPermutationDomain = TShaderPermutationDomain<FSRGBToLinear>;
+
+		BEGIN_SHADER_PARAMETER_STRUCT(FParameters, )
+			SHADER_PARAMETER_RDG_BUFFER_SRV(StructuredBuffer<FRGB16fBuffer>, InputRGB16fBuffer)
+			SHADER_PARAMETER(uint32, HorizontalElementCount)
+			SHADER_PARAMETER(uint32, VerticalElementCount)
+			SHADER_PARAMETER_RDG_TEXTURE_UAV(RWTexture2D, OutTexture)
+			END_SHADER_PARAMETER_STRUCT()
+
+	public:
+
+		// Called by the engine to determine which permutations to compile for this shader
+		static bool ShouldCompilePermutation(const FGlobalShaderPermutationParameters& Parameters)
+		{
+			return IsFeatureLevelSupported(Parameters.Platform, ERHIFeatureLevel::SM5);
+		}
+
+		/** Allocates and setup shader parameter in the incoming graph builder */
+		FRGB16fBitToRGBA16fCS::FParameters* AllocateAndSetParameters(FRDGBuilder& GraphBuilder, FRDGBufferRef RGBBuffer, FRDGTextureRef OutputTexture, int32 BufferElementsPerRow, int32 BufferLineCount);
+	};
 }
 
diff --git a/Engine/Plugins/VirtualProduction/Rivermax/RivermaxMedia/Source/RivermaxMedia/Private/RivermaxMediaCapture.cpp b/Engine/Plugins/VirtualProduction/Rivermax/RivermaxMedia/Source/RivermaxMedia/Private/RivermaxMediaCapture.cpp
--- a/Engine/Plugins/VirtualProduction/Rivermax/RivermaxMedia/Source/RivermaxMedia/Private/RivermaxMediaCapture.cpp
+++ b/Engine/Plugins/VirtualProduction/Rivermax/RivermaxMedia/Source/RivermaxMedia/Private/RivermaxMediaCapture.cpp
@@ -74,6 +74,11 @@ namespace UE::RivermaxMediaCaptureUtil
 
 		switch (InPixelFormat)
 		{
+		case ERivermaxMediaOutputPixelFormat::PF_8BIT_YUV422:
+		{
+			OutBytesPerElement = sizeof(UE::RivermaxShaders::FRGBToYUV8Bit422CS::FYUV8Bit422Buffer);
+			break;
+		}
 		case ERivermaxMediaOutputPixelFormat::PF_10BIT_YUV422:
 		{
 			OutBytesPerElement = sizeof(UE::RivermaxShaders::FRGBToYUV10Bit422LittleEndianCS::FYUV10Bit422LEBuffer);
@@ -89,6 +94,11 @@ namespace UE::RivermaxMediaCaptureUtil
 			OutBytesPerElement = sizeof(UE::RivermaxShaders::FRGBToRGB10BitCS::FRGB10BitBuffer);
 			break;
 		}
+		case ERivermaxMediaOutputPixelFormat::PF_12BIT_RGB:
+		{
+			OutBytesPerElement = sizeof(UE::RivermaxShaders::FRGBToRGB12BitCS::FRGB12BitBuffer);
+			break;
+		}
 		case ERivermaxMediaOutputPixelFormat::PF_FLOAT16_RGB:
 		{
 			OutBytesPerElement = sizeof(UE::RivermaxShaders::FRGBToRGB16fCS::FRGB16fBuffer);
@@ -261,9 +271,11 @@ EMediaCaptureResourceType URivermaxMediaCapture::GetCustomOutputResourceType() c
 
 	switch (RivermaxOutput->PixelFormat)
 	{
+	case ERivermaxMediaOutputPixelFormat::PF_8BIT_YUV422:
 	case ERivermaxMediaOutputPixelFormat::PF_10BIT_YUV422:
 	case ERivermaxMediaOutputPixelFormat::PF_8BIT_RGB:
 	case ERivermaxMediaOutputPixelFormat::PF_10BIT_RGB:
+	case ERivermaxMediaOutputPixelFormat::PF_12BIT_RGB:
 	case ERivermaxMediaOutputPixelFormat::PF_FLOAT16_RGB:
 	{
 		return EMediaCaptureResourceType::Buffer; //we use compute shader for all since output format doesn't match texture formats
@@ -298,6 +310,20 @@ void URivermaxMediaCapture::OnCustomCapture_RenderingThread(FRDGBuilder& GraphBu
 	
 	switch (RivermaxOutput->PixelFormat)
 	{
+	case ERivermaxMediaOutputPixelFormat::PF_8BIT_YUV422:
+	{
+		TShaderMapRef<FRGBToYUV8Bit422CS> ComputeShader(GlobalShaderMap);
+		FRGBToYUV8Bit422CS::FParameters* Parameters = ComputeShader->AllocateAndSetParameters(GraphBuilder, InSourceTexture, SourceSize, ViewRect, DesiredOutputSize, MediaShaders::RgbToYuvRec709Scaled, MediaShaders::YUVOffset8bits, bDoLinearToSRGB, OutputBuffer);
+
+		FComputeShaderUtils::AddPass(
+			GraphBuilder
+			, RDG_EVENT_NAME("RGBAToYUV8Bit422")
+			, ComputeShader
+			, Parameters
+			, GroupCount);
+
+		break;
+	}
 	case ERivermaxMediaOutputPixelFormat::PF_10BIT_YUV422:
 	{
 		TShaderMapRef<FRGBToYUV10Bit422LittleEndianCS> ComputeShader(GlobalShaderMap);
@@ -340,6 +366,20 @@ void URivermaxMediaCapture::OnCustomCapture_RenderingThread(FRDGBuilder& GraphBu
 
 		break;
 	}
+	case ERivermaxMediaOutputPixelFormat::PF_12BIT_RGB:
+	{
+		TShaderMapRef<FRGBToRGB12BitCS> ComputeShader(GlobalShaderMap);
+		FRGBToRGB12BitCS::FParameters* Parameters = ComputeShader->AllocateAndSetParameters(GraphBuilder, InSourceTexture, SourceSize, ViewRect, DesiredOutputSize, OutputBuffer);
+
+		FComputeShaderUtils::AddPass(
+			GraphBuilder
+			, RDG_EVENT_NAME("RGBAToRGB12Bit")
+			, ComputeShader
+			, Parameters
+			, GroupCount);
+
+		break;
+	}
 	case ERivermaxMediaOutputPixelFormat::PF_FLOAT16_RGB:
 	{
 		TShaderMapRef<FRGBToRGB16fCS> ComputeShader(GlobalShaderMap);
diff --git a/Engine/Plugins/VirtualProduction/Rivermax/RivermaxMedia/Source/RivermaxMedia/Private/RivermaxMediaOutput.cpp b/Engine/Plugins/VirtualProduction/Rivermax/RivermaxMedia/Source/RivermaxMedia/Private/RivermaxMediaOutput.cpp
--- a/Engine/Plugins/VirtualProduction/Rivermax/RivermaxMedia/Source/RivermaxMedia/Private/RivermaxMediaOutput.cpp
+++ b/Engine/Plugins/VirtualProduction/Rivermax/RivermaxMedia/Source/RivermaxMedia/Private/RivermaxMediaOutput.cpp
@@ -26,16 +26,8 @@ FIntPoint URivermaxMediaOutput::GetRequestedSize() const
 
 EPixelFormat URivermaxMediaOutput::GetRequestedPixelFormat() const
 {
+	// All output types go through buffer conversion
 	EPixelFormat Result = EPixelFormat::PF_A2B10G10R10;
-	switch (PixelFormat)
-	{
-	case ERivermaxMediaOutputPixelFormat::PF_8BIT_YUV422:
-		Result = EPixelFormat::PF_B8G8R8A8; //To be updated to use compute shader
-		break;
-	default: // All pixel formats use output buffer as their output so no need for a texture format
-		Result = EPixelFormat::PF_A2B10G10R10;
-		break;
-	}
 	return Result;
 }
 
@@ -45,11 +37,10 @@ EMediaCaptureConversionOperation URivermaxMediaOutput::GetConversionOperation(EM
 	switch (PixelFormat)
 	{
 	case ERivermaxMediaOutputPixelFormat::PF_8BIT_YUV422:
-		Result = EMediaCaptureConversionOperation::RGBA8_TO_YUV_8BIT;
-		break;
 	case ERivermaxMediaOutputPixelFormat::PF_10BIT_YUV422:
 	case ERivermaxMediaOutputPixelFormat::PF_8BIT_RGB:
 	case ERivermaxMediaOutputPixelFormat::PF_10BIT_RGB:
+	case ERivermaxMediaOutputPixelFormat::PF_12BIT_RGB:
 	case ERivermaxMediaOutputPixelFormat::PF_FLOAT16_RGB:
 	default:
 		Result = EMediaCaptureConversionOperation::CUSTOM; //We handle all conversion for rivermax since it's really tied to endianness of 2110
diff --git a/Engine/Plugins/VirtualProduction/Rivermax/RivermaxMedia/Source/RivermaxMedia/Private/RivermaxMediaTextureSample.cpp b/Engine/Plugins/VirtualProduction/Rivermax/RivermaxMedia/Source/RivermaxMedia/Private/RivermaxMediaTextureSample.cpp
--- a/Engine/Plugins/VirtualProduction/Rivermax/RivermaxMedia/Source/RivermaxMedia/Private/RivermaxMediaTextureSample.cpp
+++ b/Engine/Plugins/VirtualProduction/Rivermax/RivermaxMedia/Source/RivermaxMedia/Private/RivermaxMediaTextureSample.cpp
@@ -32,6 +32,13 @@ bool FRivermaxMediaTextureSample::ConfigureSample(uint32 InWidth, uint32 InHeigh
 	EMediaTextureSampleFormat VideoSampleFormat;
 	switch (InSampleFormat)
 	{
+		case ERivermaxMediaSourcePixelFormat::RGB_12bit:
+			// Falls through
+		case ERivermaxMediaSourcePixelFormat::RGB_16bit_Float:
+		{
+			VideoSampleFormat = EMediaTextureSampleFormat::FloatRGBA;
+			break;
+		}
 		case ERivermaxMediaSourcePixelFormat::RGB_10bit:
 			// Falls through
 		case ERivermaxMediaSourcePixelFormat::YUV422_10bit:
diff --git a/Engine/Plugins/VirtualProduction/Rivermax/RivermaxMedia/Source/RivermaxMedia/Private/RivermaxMediaTextureSampleConverter.cpp b/Engine/Plugins/VirtualProduction/Rivermax/RivermaxMedia/Source/RivermaxMedia/Private/RivermaxMediaTextureSampleConverter.cpp
--- a/Engine/Plugins/VirtualProduction/Rivermax/RivermaxMedia/Source/RivermaxMedia/Private/RivermaxMediaTextureSampleConverter.cpp
+++ b/Engine/Plugins/VirtualProduction/Rivermax/RivermaxMedia/Source/RivermaxMedia/Private/RivermaxMediaTextureSampleConverter.cpp
@@ -5,7 +5,6 @@
 #include "RenderGraphBuilder.h"
 #include "RivermaxShaders.h"
 #include "RivermaxMediaTextureSample.h"
-#include "RivermaxMediaUtils.h"
 
 
 void FRivermaxMediaTextureSampleConverter::Setup(ERivermaxMediaSourcePixelFormat InPixelFormat, TWeakPtr<FRivermaxMediaTextureSample> InSample, bool bInDoSRGBToLinear)
@@ -20,7 +19,6 @@ bool FRivermaxMediaTextureSampleConverter::Convert(FTexture2DRHIRef& InDestinati
 	TRACE_CPUPROFILER_EVENT_SCOPE(RivermaxSampleConverter::Convert);
 
 	using namespace UE::RivermaxShaders;
-	using namespace UE::RivermaxCore;
 
 	TSharedPtr<FRivermaxMediaTextureSample> SamplePtr = Sample.Pin();
 	if (SamplePtr.IsValid() == false)
@@ -31,6 +29,8 @@ bool FRivermaxMediaTextureSampleConverter::Convert(FTexture2DRHIRef& InDestinati
 	FIntVector GroupCount;
 	
 	FRDGBuilder GraphBuilder(FRHICommandListExecutor::GetImmediateCommandList());
+	SCOPED_DRAW_EVENT(GraphBuilder.RHICmdList, Rivermax_SampleConverter);
+	
 	FRDGTextureRef OutputResource = GraphBuilder.RegisterExternalTexture(CreateRenderTarget(InDestinationTexture, TEXT("RivermaxMediaTextureOutputResource")));
 
 	//Configure shader and add conversion pass based on desired pixel format
@@ -38,13 +38,40 @@ bool FRivermaxMediaTextureSampleConverter::Convert(FTexture2DRHIRef& InDestinati
 	{
 	case ERivermaxMediaSourcePixelFormat::YUV422_8bit:
 	{
+		TRACE_CPUPROFILER_EVENT_SCOPE(RivermaxSampleConverter::YUV8ShaderSetup);
+		const int32 BytesPerElement = sizeof(FYUV8Bit422ToRGBACS::FYUV8Bit422Buffer);
+		const uint32 Stride = SamplePtr->GetStride();
+		const int32 ElementsPerRow = (Stride / BytesPerElement) + ((Stride % BytesPerElement > 0) ? 1 : 0);
+		const int32 ElementCount = ElementsPerRow * InDestinationTexture->GetDesc().Extent.Y;
+
+		FYUV8Bit422ToRGBACS::FPermutationDomain PermutationVector;
+		PermutationVector.Set<FYUV8Bit422ToRGBACS::FSRGBToLinear>(bDoSRGBToLinear);
+
+		FRDGBufferRef InputYUVBuffer = CreateStructuredBuffer(GraphBuilder, TEXT("RivermaxInputBuffer"), BytesPerElement, ElementCount, SamplePtr->GetBuffer(), BytesPerElement * ElementCount);
+		constexpr int32 PixelsPerInput = 2;
+		const FIntPoint ProcessedOutputDimension = { InDestinationTexture->GetDesc().Extent.X / PixelsPerInput, InDestinationTexture->GetDesc().Extent.Y };
+		GroupCount = FComputeShaderUtils::GetGroupCount(ProcessedOutputDimension, FComputeShaderUtils::kGolden2DGroupSize);
+		FGlobalShaderMap* GlobalShaderMap = GetGlobalShaderMap(GMaxRHIFeatureLevel);
+		TShaderMapRef<FYUV8Bit422ToRGBACS> ComputeShader(GlobalShaderMap, PermutationVector);
+		FMatrix YUVToRGBMatrix = SamplePtr->GetYUVToRGBMatrix();
+		FVector YUVOffset(MediaShaders::YUVOffset8bits);
+		FYUV8Bit422ToRGBACS::FParameters* Parameters = ComputeShader->AllocateAndSetParameters(GraphBuilder, InputYUVBuffer, OutputResource, YUVToRGBMatrix, YUVOffset, ElementsPerRow, ProcessedOutputDimension.Y);
+
+
+		FComputeShaderUtils::AddPass(
+			GraphBuilder
+			, RDG_EVENT_NAME("YUV8Bit422ToRGBA")
+			, ComputeShader
+			, Parameters
+			, GroupCount);
 		break;
 	}
 	case ERivermaxMediaSourcePixelFormat::YUV422_10bit:
 	{
 		TRACE_CPUPROFILER_EVENT_SCOPE(RivermaxSampleConverter::YUV10ShaderSetup);
 		const int32 BytesPerElement = sizeof(FYUV10Bit422ToRGBACS::FYUV10Bit422LEBuffer);
-		const int32 ElementsPerRow = FMath::CeilToInt32(SamplePtr->GetStride() / (float)BytesPerElement);
+		const uint32 Stride = SamplePtr->GetStride();
+		const int32 ElementsPerRow = (Stride / BytesPerElement) + ((Stride % BytesPerElement > 0) ? 1 : 0);
 		const int32 ElementCount = ElementsPerRow * InDestinationTexture->GetDesc().Extent.Y;
 
 		FYUV10Bit422ToRGBACS::FPermutationDomain PermutationVector;
@@ -73,7 +100,8 @@ bool FRivermaxMediaTextureSampleConverter::Convert(FTexture2DRHIRef& InDestinati
 	{
 		TRACE_CPUPROFILER_EVENT_SCOPE(RivermaxSampleConverter::RGB8ShaderSetup);
 		const int32 BytesPerElement = sizeof(FRGB8BitToRGBA8CS::FRGB8BitBuffer);
-		const int32 ElementsPerRow = FMath::CeilToInt32(SamplePtr->GetStride() / (float)BytesPerElement);
+		const uint32 Stride = SamplePtr->GetStride();
+		const int32 ElementsPerRow = (Stride / BytesPerElement) + ((Stride % BytesPerElement > 0) ? 1 : 0);
 		const int32 ElementCount = ElementsPerRow * InDestinationTexture->GetDesc().Extent.Y;
 
 		FRGB8BitToRGBA8CS::FPermutationDomain PermutationVector;
@@ -100,8 +128,9 @@ bool FRivermaxMediaTextureSampleConverter::Convert(FTexture2DRHIRef& InDestinati
 	{
 		TRACE_CPUPROFILER_EVENT_SCOPE(RivermaxSampleConverter::RGB10ShaderSetup);
 
-		const int32 BytesPerElement = sizeof(FRGBToRGB10BitCS::FRGB10BitBuffer);
-		const int32 ElementsPerRow = FMath::CeilToInt32(SamplePtr->GetStride() / (float)BytesPerElement);
+		const int32 BytesPerElement = sizeof(FRGB10BitToRGBA10CS::FRGB10BitBuffer);
+		const uint32 Stride = SamplePtr->GetStride();
+		const int32 ElementsPerRow = (Stride / BytesPerElement) + ((Stride % BytesPerElement > 0) ? 1 : 0);
 		const int32 ElementCount = ElementsPerRow * InDestinationTexture->GetDesc().Extent.Y;
 		
 		FRGB10BitToRGBA10CS::FPermutationDomain PermutationVector;
@@ -124,6 +153,64 @@ bool FRivermaxMediaTextureSampleConverter::Convert(FTexture2DRHIRef& InDestinati
 			, GroupCount);
 		break;
 	}
+	case ERivermaxMediaSourcePixelFormat::RGB_12bit:
+	{
+		TRACE_CPUPROFILER_EVENT_SCOPE(RivermaxSampleConverter::RGB12ShaderSetup);
+
+		const int32 BytesPerElement = sizeof(FRGB12BitToRGBA12CS::FRGB12BitBuffer);
+		const uint32 Stride = SamplePtr->GetStride();
+		const int32 ElementsPerRow = (Stride / BytesPerElement) + ((Stride % BytesPerElement > 0) ? 1 : 0);
+		const int32 ElementCount = ElementsPerRow * InDestinationTexture->GetDesc().Extent.Y;
+
+		FRGB12BitToRGBA12CS::FPermutationDomain PermutationVector;
+		PermutationVector.Set<FRGB12BitToRGBA12CS::FSRGBToLinear>(bDoSRGBToLinear);
+
+		FRDGBufferRef InputRGGBuffer = CreateStructuredBuffer(GraphBuilder, TEXT("RivermaxInputBuffer"), BytesPerElement, ElementCount, SamplePtr->GetBuffer(), BytesPerElement * ElementCount);
+		constexpr int32 PixelsPerInput = 8;
+		const FIntPoint ProcessedOutputDimension = { InDestinationTexture->GetDesc().Extent.X / PixelsPerInput,InDestinationTexture->GetDesc().Extent.Y };
+		GroupCount = FComputeShaderUtils::GetGroupCount(ProcessedOutputDimension, FComputeShaderUtils::kGolden2DGroupSize);
+		FGlobalShaderMap* GlobalShaderMap = GetGlobalShaderMap(GMaxRHIFeatureLevel);
+		TShaderMapRef<FRGB12BitToRGBA12CS> ComputeShader(GlobalShaderMap, PermutationVector);
+		FRGB12BitToRGBA12CS::FParameters* Parameters = ComputeShader->AllocateAndSetParameters(GraphBuilder, InputRGGBuffer, OutputResource, ElementsPerRow, ProcessedOutputDimension.Y);
+
+
+		FComputeShaderUtils::AddPass(
+			GraphBuilder
+			, RDG_EVENT_NAME("RGB12BitToRGBA")
+			, ComputeShader
+			, Parameters
+			, GroupCount);
+		break;
+	}
+	case ERivermaxMediaSourcePixelFormat::RGB_16bit_Float:
+	{
+		TRACE_CPUPROFILER_EVENT_SCOPE(RivermaxSampleConverter::RGB16FloatShaderSetup);
+
+		const int32 BytesPerElement = sizeof(FRGB16fBitToRGBA16fCS::FRGB16fBuffer);
+		const uint32 Stride = SamplePtr->GetStride();
+		const int32 ElementsPerRow = (Stride / BytesPerElement) + ((Stride % BytesPerElement > 0) ? 1 : 0);
+		const int32 ElementCount = ElementsPerRow * InDestinationTexture->GetDesc().Extent.Y;
+
+		FRGB16fBitToRGBA16fCS::FPermutationDomain PermutationVector;
+		PermutationVector.Set<FRGB16fBitToRGBA16fCS::FSRGBToLinear>(bDoSRGBToLinear);
+
+		FRDGBufferRef InputRGGBuffer = CreateStructuredBuffer(GraphBuilder, TEXT("RivermaxInputBuffer"), BytesPerElement, ElementCount, SamplePtr->GetBuffer(), BytesPerElement * ElementCount);
+		constexpr int32 PixelsPerInput = 2;
+		const FIntPoint ProcessedOutputDimension = { InDestinationTexture->GetDesc().Extent.X / PixelsPerInput,InDestinationTexture->GetDesc().Extent.Y };
+		GroupCount = FComputeShaderUtils::GetGroupCount(ProcessedOutputDimension, FComputeShaderUtils::kGolden2DGroupSize);
+		FGlobalShaderMap* GlobalShaderMap = GetGlobalShaderMap(GMaxRHIFeatureLevel);
+		TShaderMapRef<FRGB16fBitToRGBA16fCS> ComputeShader(GlobalShaderMap, PermutationVector);
+		FRGB16fBitToRGBA16fCS::FParameters* Parameters = ComputeShader->AllocateAndSetParameters(GraphBuilder, InputRGGBuffer, OutputResource, ElementsPerRow, ProcessedOutputDimension.Y);
+
+
+		FComputeShaderUtils::AddPass(
+			GraphBuilder
+			, RDG_EVENT_NAME("RGB16fBitToRGBA")
+			, ComputeShader
+			, Parameters
+			, GroupCount);
+		break;
+	}
 	default:
 	{
 		ensureMsgf(false, TEXT("Unhandled pixel format (%d) given to Rivermax MediaSample converter"), InputPixelFormat);
diff --git a/Engine/Plugins/VirtualProduction/Rivermax/RivermaxMedia/Source/RivermaxMedia/Private/RivermaxMediaUtils.cpp b/Engine/Plugins/VirtualProduction/Rivermax/RivermaxMedia/Source/RivermaxMedia/Private/RivermaxMediaUtils.cpp
--- a/Engine/Plugins/VirtualProduction/Rivermax/RivermaxMedia/Source/RivermaxMedia/Private/RivermaxMediaUtils.cpp
+++ b/Engine/Plugins/VirtualProduction/Rivermax/RivermaxMedia/Source/RivermaxMedia/Private/RivermaxMediaUtils.cpp
@@ -26,6 +26,10 @@ namespace UE::RivermaxMediaUtils::Private
 		{
 			return ESamplingType::RGB_10bit;
 		}
+		case ERivermaxMediaOutputPixelFormat::PF_12BIT_RGB:
+		{
+			return ESamplingType::RGB_12bit;
+		}
 		case ERivermaxMediaOutputPixelFormat::PF_FLOAT16_RGB:
 		{
 			return ESamplingType::RGB_16bitFloat;
@@ -60,6 +64,14 @@ namespace UE::RivermaxMediaUtils::Private
 		{
 			return ESamplingType::RGB_10bit;
 		}
+		case ERivermaxMediaSourcePixelFormat::RGB_12bit:
+		{
+			return ESamplingType::RGB_12bit;
+		}
+		case ERivermaxMediaSourcePixelFormat::RGB_16bit_Float:
+		{
+			return ESamplingType::RGB_16bitFloat;
+		}
 		default:
 		{
 			checkNoEntry();
diff --git a/Engine/Plugins/VirtualProduction/Rivermax/RivermaxMedia/Source/RivermaxMedia/Public/RivermaxMediaOutput.h b/Engine/Plugins/VirtualProduction/Rivermax/RivermaxMedia/Source/RivermaxMedia/Public/RivermaxMediaOutput.h
--- a/Engine/Plugins/VirtualProduction/Rivermax/RivermaxMedia/Source/RivermaxMedia/Public/RivermaxMediaOutput.h
+++ b/Engine/Plugins/VirtualProduction/Rivermax/RivermaxMedia/Source/RivermaxMedia/Public/RivermaxMediaOutput.h
@@ -17,6 +17,7 @@ enum class ERivermaxMediaOutputPixelFormat : uint8
 	PF_10BIT_YUV422 UMETA(DisplayName = "10bit YUV422"),
 	PF_8BIT_RGB UMETA(DisplayName = "8bit RGB"),
 	PF_10BIT_RGB UMETA(DisplayName = "10bit RGB"),
+	PF_12BIT_RGB UMETA(DisplayName = "12bit RGB"),
 	PF_FLOAT16_RGB UMETA(DisplayName = "16bit Float RGB")
 };
 
diff --git a/Engine/Plugins/VirtualProduction/Rivermax/RivermaxMedia/Source/RivermaxMedia/Public/RivermaxMediaSource.h b/Engine/Plugins/VirtualProduction/Rivermax/RivermaxMedia/Source/RivermaxMedia/Public/RivermaxMediaSource.h
--- a/Engine/Plugins/VirtualProduction/Rivermax/RivermaxMedia/Source/RivermaxMedia/Public/RivermaxMediaSource.h
+++ b/Engine/Plugins/VirtualProduction/Rivermax/RivermaxMedia/Source/RivermaxMedia/Public/RivermaxMediaSource.h
@@ -18,7 +18,8 @@ enum class ERivermaxMediaSourcePixelFormat : uint8
 	YUV422_10bit UMETA(DisplayName = "10bit YUV422"),
 	RGB_8bit UMETA(DisplayName = "8bit RGB"),
 	RGB_10bit UMETA(DisplayName = "10bit RGB"),
-	//Todo Add support for float16 rgb
+	RGB_12bit UMETA(DisplayName = "12bit RGB"),
+	RGB_16bit_Float UMETA(DisplayName = "16bit Float RGB"),
 };
 
 /**
