- Adding support for GPUDirect capability verification in rivermax
- Adding support for gpudirect for rivermax output leveraging cuda
- Adding support for clock fallback if PTP fails
- Changing max payload size to 1420 to respect standard

diff --git a/Engine/Plugins/VirtualProduction/Rivermax/RivermaxCore/Source/RivermaxCore/Private/RivermaxManager.cpp b/Engine/Plugins/VirtualProduction/Rivermax/RivermaxCore/Source/RivermaxCore/Private/RivermaxManager.cpp
--- a/Engine/Plugins/VirtualProduction/Rivermax/RivermaxCore/Source/RivermaxCore/Private/RivermaxManager.cpp
+++ b/Engine/Plugins/VirtualProduction/Rivermax/RivermaxCore/Source/RivermaxCore/Private/RivermaxManager.cpp
@@ -2,6 +2,7 @@
 
 #include "RivermaxManager.h"
 
+#include "CudaModule.h"
 #include "Interfaces/IPv4/IPv4Endpoint.h"
 #include "IRivermaxCoreModule.h"
 #include "Misc/App.h"
@@ -27,6 +28,12 @@
 
 namespace UE::RivermaxCore::Private
 {
+	static TAutoConsoleVariable<int32> CVarRivermaxEnableGPUDirectCapability(
+		TEXT("Rivermax.GPUDirect"), 1,
+		TEXT("Whether GPUDirect capability is validated at launch. Can be used to disable it entirely."),
+		ECVF_Default);
+
+
 	static uint64 RivermaxTimeHandler(void* Context)
 	{
 		IRivermaxCoreModule* RivermaxModule = FModuleManager::GetModulePtr<IRivermaxCoreModule>("RivermaxCore");
@@ -41,6 +48,12 @@ namespace UE::RivermaxCore::Private
 	FRivermaxManager::FRivermaxManager()
 		: DeviceFinder(MakeUnique<FRivermaxDeviceFinder>())
 	{
+		// No need loading Cuda if we won't be looking at gpudirect capabilities
+		if (CVarRivermaxEnableGPUDirectCapability.GetValueOnGameThread())
+		{
+			FModuleManager::LoadModuleChecked<FCUDAModule>("CUDA");
+		}
+
 		const bool bIsLibraryLoaded = LoadRivermaxLibrary();
 		if (bIsLibraryLoaded && FApp::CanEverRender())
 		{
@@ -128,6 +141,15 @@ namespace UE::RivermaxCore::Private
 		rmax_init_config Config;
 		memset(&Config, 0, sizeof(Config));
 		Config.flags |= RIVERMAX_ENABLE_CLOCK_CONFIGURATION;
+
+		ERivermaxTimeSource ConfiguredSource = ERivermaxTimeSource::Platform;
+		const auto ConfigurePlatformClock = [&ConfiguredSource](rmax_init_config& InOutConfig)
+		{
+			InOutConfig.clock_configurations.clock_type = rmax_clock_types::RIVERMAX_USER_CLOCK_HANDLER;
+			InOutConfig.clock_configurations.clock_u.rmax_user_clock_handler.clock_handler = RivermaxTimeHandler;
+			InOutConfig.clock_configurations.clock_u.rmax_user_clock_handler.ctx = nullptr;
+			ConfiguredSource = ERivermaxTimeSource::Platform;
+		};
 			
 		const URivermaxSettings* Settings = GetDefault<URivermaxSettings>();
 
@@ -142,7 +164,7 @@ namespace UE::RivermaxCore::Private
 				if (inet_pton(AF_INET, StringCast<ANSICHAR>(*PTPAddress).Get(), &Config.clock_configurations.clock_u.rmax_ptp_clock.device_ip_addr))
 				{
 					UE_LOG(LogRivermax, Display, TEXT("Initialize Rivermax clock as PTP using device interface %s"), *PTPAddress);
-					TimeSource = ERivermaxTimeSource::PTP;
+					ConfiguredSource = ERivermaxTimeSource::PTP;
 					bConfiguredSuccessfully = true;
 				}
 			}
@@ -153,16 +175,22 @@ namespace UE::RivermaxCore::Private
 			}
 		}
 
-		if(bConfiguredSuccessfully == false)
+		if (bConfiguredSuccessfully == false)
 		{
-			// Todo: Need to update clock fetcher to align with UTC. Refer to SDK
-			Config.clock_configurations.clock_type = rmax_clock_types::RIVERMAX_USER_CLOCK_HANDLER;
-			Config.clock_configurations.clock_u.rmax_user_clock_handler.clock_handler = RivermaxTimeHandler;
-			Config.clock_configurations.clock_u.rmax_user_clock_handler.ctx = nullptr;
-			TimeSource = ERivermaxTimeSource::Platform;
+			ConfigurePlatformClock(Config);
 		}
 
 		rmax_status_t Status = rmax_init(&Config);
+		if (Status == RMAX_ERR_CLOCK_TYPE_NOT_SUPPORTED)
+		{
+			if (Settings->TimeSource == ERivermaxTimeSource::PTP)
+			{
+				UE_LOG(LogRivermax, Warning, TEXT("Failed to configure Rivermax using PTP clock. Falling back to platform clock."));
+				ConfigurePlatformClock(Config);
+				Status = rmax_init(&Config);
+			}
+		}
+
 		if (Status == RMAX_OK)
 		{
 			uint32 Major = 0;
@@ -172,8 +200,15 @@ namespace UE::RivermaxCore::Private
 			Status = rmax_get_version(&Major, &Minor, &Release, &Build);
 			if (Status == RMAX_OK)
 			{
+				if (CVarRivermaxEnableGPUDirectCapability.GetValueOnGameThread())
+				{
+					VerifyGPUDirectCapability();
+				}
+
 				bIsInitialized = true;
-				UE_LOG(LogRivermax, Log, TEXT("Rivermax library version %d.%d.%d.%d succesfully initialized"), Major, Minor, Release, Build);
+				TimeSource = ConfiguredSource;
+				const TCHAR* GPUDirectSupport = bIsGPUDirectSupported ? TEXT("with GPUDirect support.") : TEXT("without GPUDirect support.");
+				UE_LOG(LogRivermax, Log, TEXT("Rivermax library version %d.%d.%d.%d succesfully initialized, %s"), Major, Minor, Release, Build, GPUDirectSupport);
 			}
 			else
 			{
@@ -210,6 +245,61 @@ namespace UE::RivermaxCore::Private
 	{
 		return bIsGPUDirectSupported;
 	}
+
+	void FRivermaxManager::VerifyGPUDirectCapability()
+	{
+		bIsGPUDirectSupported = false;
+
+		const ERHIInterfaceType RHIType = RHIGetInterfaceType();
+		if (RHIType != ERHIInterfaceType::D3D12)
+		{
+			UE_LOG(LogRivermax, Log, TEXT("GPUDirect won't be available for Rivermax as it's only available with D3D12 RHI."));
+			return;
+		}
+		
+		FCUDAModule* CudaModule = FModuleManager::GetModulePtr<FCUDAModule>("CUDA");
+		if (CudaModule)
+		{
+			CudaModule->DriverAPI()->cuCtxPushCurrent(CudaModule->GetCudaContext());
+
+			int DeviceCount = -1;
+			CUresult Result = CudaModule->DriverAPI()->cuDeviceGetCount(&DeviceCount);
+			if (DeviceCount <= 0)
+			{
+				UE_LOG(LogRivermax, Log, TEXT("No Cuda compatible device found. GPUDirect won't be available for Rivermax. Cuda status: %d"), Result);
+				return;
+			}
+
+			// todo: add support for mgpu compatibility verification
+			const int GPUIndex = 0;
+			CUdevice CudaDevice;
+			Result = CudaModule->DriverAPI()->cuDeviceGet(&CudaDevice, GPUIndex);
+			if (Result != CUDA_SUCCESS)
+			{
+				UE_LOG(LogRivermax, Log, TEXT("Could not get a valid Cuda device. GPUDirect won't be available for Rivermax. Cuda status: %d"), Result);
+				return;
+			}
+
+
+			int SupportsVMM = 0;
+			Result  = CudaModule->DriverAPI()->cuDeviceGetAttribute(&SupportsVMM, CU_DEVICE_ATTRIBUTE_VIRTUAL_MEMORY_MANAGEMENT_SUPPORTED, CudaDevice);
+			if (SupportsVMM != 1)
+			{
+				UE_LOG(LogRivermax, Log, TEXT("Cuda device doesn't support virtual memory management. GPUDirect won't be available for Rivermax. Cuda status: %d"), Result);
+				return;
+			}
+
+			int SupportsRDMA = 0;
+			Result = CudaModule->DriverAPI()->cuDeviceGetAttribute(&SupportsRDMA, CU_DEVICE_ATTRIBUTE_GPU_DIRECT_RDMA_WITH_CUDA_VMM_SUPPORTED, CudaDevice);
+			if (SupportsRDMA != 1)
+			{
+				UE_LOG(LogRivermax, Log, TEXT("Cuda device doesn't support RDMA. GPUDirect won't be available for Rivermax. Cuda status: %d"), Result);
+				return;
+			}
+
+			bIsGPUDirectSupported = true;
+		}
+	}
 }
 
 
diff --git a/Engine/Plugins/VirtualProduction/Rivermax/RivermaxCore/Source/RivermaxCore/Private/RivermaxManager.h b/Engine/Plugins/VirtualProduction/Rivermax/RivermaxCore/Source/RivermaxCore/Private/RivermaxManager.h
--- a/Engine/Plugins/VirtualProduction/Rivermax/RivermaxCore/Source/RivermaxCore/Private/RivermaxManager.h
+++ b/Engine/Plugins/VirtualProduction/Rivermax/RivermaxCore/Source/RivermaxCore/Private/RivermaxManager.h
@@ -28,6 +28,7 @@ namespace UE::RivermaxCore::Private
 	private:
 		bool LoadRivermaxLibrary();
 		void InitializeLibrary();
+		void VerifyGPUDirectCapability();
 
 	private:
 
diff --git a/Engine/Plugins/VirtualProduction/Rivermax/RivermaxCore/Source/RivermaxCore/Private/RivermaxOutputFrame.cpp b/Engine/Plugins/VirtualProduction/Rivermax/RivermaxCore/Source/RivermaxCore/Private/RivermaxOutputFrame.cpp
--- a/Engine/Plugins/VirtualProduction/Rivermax/RivermaxCore/Source/RivermaxCore/Private/RivermaxOutputFrame.cpp
+++ b/Engine/Plugins/VirtualProduction/Rivermax/RivermaxCore/Source/RivermaxCore/Private/RivermaxOutputFrame.cpp
@@ -4,18 +4,19 @@
 
 namespace UE::RivermaxCore::Private
 {
-	FRivermaxOutputFrame::FRivermaxOutputFrame(uint32 InFrameIndex)
+	FRivermaxOutputFrame::FRivermaxOutputFrame(uint32 InFrameIndex, TFunction<void(void*)> InDeallocationFunction)
 		: FrameIndex(InFrameIndex)
 		, VideoBuffer(nullptr)
+		, DeallocationFunc(InDeallocationFunction)
 	{
 		Clear();
 	}
 
 	FRivermaxOutputFrame::~FRivermaxOutputFrame()
 	{
-		if (VideoBuffer)
+		if (DeallocationFunc)
 		{
-			FMemory::Free(VideoBuffer);
+			DeallocationFunc(VideoBuffer);
 		}
 	}
 
diff --git a/Engine/Plugins/VirtualProduction/Rivermax/RivermaxCore/Source/RivermaxCore/Private/RivermaxOutputFrame.h b/Engine/Plugins/VirtualProduction/Rivermax/RivermaxCore/Source/RivermaxCore/Private/RivermaxOutputFrame.h
--- a/Engine/Plugins/VirtualProduction/Rivermax/RivermaxCore/Source/RivermaxCore/Private/RivermaxOutputFrame.h
+++ b/Engine/Plugins/VirtualProduction/Rivermax/RivermaxCore/Source/RivermaxCore/Private/RivermaxOutputFrame.h
@@ -14,7 +14,7 @@ namespace UE::RivermaxCore::Private
 	 */
 	struct FRivermaxOutputFrame
 	{
-		FRivermaxOutputFrame(uint32 InFrameIndex);
+		FRivermaxOutputFrame(uint32 InFrameIndex, TFunction<void(void*)> InDeallocationFunction);
 		~FRivermaxOutputFrame();
 		FRivermaxOutputFrame(const FRivermaxOutputFrame&) = delete;
 		FRivermaxOutputFrame& operator=(const FRivermaxOutputFrame&) = delete;
@@ -31,8 +31,11 @@ namespace UE::RivermaxCore::Private
 		/** Index of the frame used as unique identifier */
 		uint32 FrameIndex = 0;
 
-		/** Video buffer where we copy texture to be sent out */
-		uint8* VideoBuffer = nullptr;
+		/** 
+		 * Video buffer where we copy texture to be sent out 
+		 * If using GPUDirect, memory will be allocated in cuda space
+		 */
+		void* VideoBuffer = nullptr;
 
 		/** Identifier of the frame incremented every time a frame is captured */
 		static constexpr uint32 InvalidIdentifier = MAX_uint32;
@@ -57,7 +60,10 @@ namespace UE::RivermaxCore::Private
 		void* PayloadPtr = nullptr;
 		void* HeaderPtr = nullptr;
 
-		
+	private:
+
+		/** Method called to deallocate memory. */
+		TFunction<void(void*)> DeallocationFunc;
 	};
 }
 
diff --git a/Engine/Plugins/VirtualProduction/Rivermax/RivermaxCore/Source/RivermaxCore/Private/RivermaxOutputStream.cpp b/Engine/Plugins/VirtualProduction/Rivermax/RivermaxCore/Source/RivermaxCore/Private/RivermaxOutputStream.cpp
--- a/Engine/Plugins/VirtualProduction/Rivermax/RivermaxCore/Source/RivermaxCore/Private/RivermaxOutputStream.cpp
+++ b/Engine/Plugins/VirtualProduction/Rivermax/RivermaxCore/Source/RivermaxCore/Private/RivermaxOutputStream.cpp
@@ -3,6 +3,8 @@
 #include "RivermaxOutputStream.h"
 
 #include "Async/Async.h"
+#include "CudaModule.h"
+#include "ID3D12DynamicRHI.h"
 #include "IRivermaxCoreModule.h"
 #include "IRivermaxManager.h"
 #include "Misc/ByteSwap.h"
@@ -11,7 +13,6 @@
 #include "RivermaxTypes.h"
 #include "RivermaxUtils.h"
 
-
 namespace UE::RivermaxCore::Private
 {
 	static TAutoConsoleVariable<int32> CVarRivermaxWakeupOffset(
@@ -132,7 +133,7 @@ namespace UE::RivermaxCore::Private
 		// We do (try) to make gpu allocations here to let the capturer know if we require it or not.
 		if (RivermaxModule.GetRivermaxManager()->IsGPUDirectSupported() && Options.bUseGPUDirect)
 		{
-			//bUseGPUDirect = AllocateGPUBuffers();
+			bUseGPUDirect = AllocateGPUBuffers();
 		}
 
 		Async(EAsyncExecution::TaskGraph, [this]()
@@ -215,6 +216,8 @@ namespace UE::RivermaxCore::Private
 			FPlatformProcess::ReturnSynchEventToPool(ReadyToSendEvent);
 			ReadyToSendEvent = nullptr;
 			UE_LOG(LogRivermax, Log, TEXT("Rivermax Output stream has shutdown"));
+
+			DeallocateBuffers();
 		}
 	}
 
@@ -316,12 +319,18 @@ namespace UE::RivermaxCore::Private
 		const int32 Stride = GetStride();
 		const int32 FrameAllocSize = Options.AlignedResolution.Y * Stride;
 		AvailableFrames.Reserve(Options.NumberOfBuffers);
+
+		const auto SystemDeallocator = [](void* Buffer)
+		{
+			FMemory::Free(Buffer);
+		};
+
 		for (int32 Index = 0; Index < Options.NumberOfBuffers; ++Index)
 		{
-			TSharedPtr<FRivermaxOutputFrame> Frame = MakeShared<FRivermaxOutputFrame>(Index);
+			TSharedPtr<FRivermaxOutputFrame> Frame = MakeShared<FRivermaxOutputFrame>(Index, SystemDeallocator);
 
 			constexpr uint32 CacheLineSize = PLATFORM_CACHE_LINE_SIZE;
-			Frame->VideoBuffer = static_cast<uint8*>(FMemory::Malloc(FrameAllocSize, CacheLineSize));
+			Frame->VideoBuffer = FMemory::Malloc(FrameAllocSize, CacheLineSize);
 			AvailableFrames.Add(MoveTemp(Frame));
 		}
 	}
@@ -963,6 +972,74 @@ namespace UE::RivermaxCore::Private
 	{
 		TRACE_CPUPROFILER_EVENT_SCOPE(FRivermaxOutputStream::PushGPUVideoFrame);
 
+		const int32 Stride = GetStride();
+		if (TSharedPtr<FRivermaxOutputFrame> AvailableFrame = GetNextAvailableFrame(NewFrame.FrameIdentifier))
+		{
+			const FString TraceName = FString::Format(TEXT("FRivermaxOutputStream::PushFrame {0}"), { AvailableFrame->FrameIndex });
+			TRACE_CPUPROFILER_EVENT_SCOPE_TEXT(*TraceName);
+
+			FCUDAModule& CudaModule = FModuleManager::GetModuleChecked<FCUDAModule>("CUDA");
+			CUresult Result = CudaModule.DriverAPI()->cuCtxPushCurrent(CudaModule.GetCudaContext());
+			
+			void* MappedPointer = GetMappedAddress(CapturedBuffer);
+			if (MappedPointer == nullptr)
+			{
+				UE_LOG(LogRivermax, Error, TEXT("Failed to find a mapped memory address for captured buffer. Stopping capture."));
+				Listener->OnStreamError();
+				return false;
+			}
+
+			const CUdeviceptr CudaMemoryPointer = reinterpret_cast<CUdeviceptr>(MappedPointer);
+			Result = CudaModule.DriverAPI()->cuMemcpyDtoD(reinterpret_cast<CUdeviceptr>(AvailableFrame->VideoBuffer), CudaMemoryPointer, Options.AlignedResolution.Y * Stride);
+			if (Result != CUDA_SUCCESS)
+			{
+				UE_LOG(LogRivermax, Error, TEXT("Failed to copy captured bufer to cuda memory. Stopping capture. Error: %d"), Result);
+				Listener->OnStreamError();
+				return false;
+			}
+
+
+			// Callback called by Cuda when stream work has completed on cuda engine (MemCpy -> Callback)
+			// Once Memcpy has been done, we know we can mark that memory as available to be sent. 
+			
+			auto CudaCallback = [](CUstream hStream, CUresult status, void* userData)
+			{
+				TRACE_CPUPROFILER_EVENT_SCOPE(CudaWorkDoneCallback);
+
+				FRivermaxOutputStream* Stream = reinterpret_cast<FRivermaxOutputStream*>(userData);
+
+				TOptional<uint32> TargetFrameIdentifier = Stream->PendingIdentifiers.Dequeue();
+				if (TargetFrameIdentifier.IsSet())
+				{
+					if (TSharedPtr<FRivermaxOutputFrame> AvailableFrame = Stream->GetNextAvailableFrame(TargetFrameIdentifier.GetValue()))
+					{
+						AvailableFrame->bIsVideoBufferReady = true;
+						if (AvailableFrame->IsReadyToBeSent())
+						{
+							const FString TraceName = FString::Format(TEXT("FRivermaxOutputStream::PushFrame {0}"), { AvailableFrame->FrameIndex });
+							TRACE_CPUPROFILER_EVENT_SCOPE_TEXT(*TraceName);
+
+							// Make frame available to be sent
+							FScopeLock Lock(&Stream->FrameCriticalSection);
+							Stream->AvailableFrames.Remove(AvailableFrame);
+							Stream->FramesToSend.Add(MoveTemp(AvailableFrame));
+							Stream->ReadyToSendEvent->Trigger();
+						}
+					}
+				}
+			};
+			
+			// Add pending frame for cuda callback 
+			PendingIdentifiers.Enqueue(AvailableFrame->FrameIdentifier);
+
+			// Schedule a callback to make the frame available
+			CudaModule.DriverAPI()->cuStreamAddCallback(0, CudaCallback, this, 0);
+			
+			FCUDAModule::CUDA().cuCtxPopCurrent(nullptr);
+
+			return true;
+		}
+
 		return false;
 	}
 
@@ -971,10 +1048,256 @@ namespace UE::RivermaxCore::Private
 		return bUseGPUDirect;
 	}
 
+	bool FRivermaxOutputStream::AllocateGPUBuffers()
+	{
+		// Allocate a single memory space that will contain all frame buffers
+		
+		TRACE_CPUPROFILER_EVENT_SCOPE(FRivermaxOutputStream::AllocateGPUBuffers);
+
+		const ERHIInterfaceType RHIType = RHIGetInterfaceType();
+		if (RHIType != ERHIInterfaceType::D3D12)
+		{
+			UE_LOG(LogRivermax, Warning, TEXT("Can't initialize output to use GPUDirect. RHI is %d but only Dx12 is supported at the moment."), RHIType);
+			return false;
+		}
+
+		const int32 Stride = GetStride();
+		const int32 FrameAllocSize = Options.AlignedResolution.Y * Stride;
+
+		FCUDAModule& CudaModule = FModuleManager::GetModuleChecked<FCUDAModule>("CUDA");
+		CudaModule.DriverAPI()->cuCtxPushCurrent(CudaModule.GetCudaContext());
+
+		// Todo: Add support for mgpu. For now, this will not work unless the memcpy does implicitely a cross gpu transfer.
+		constexpr int GPUIndex = 0;
+		CUdevice CudaDevice;
+		CUresult Status = CudaModule.DriverAPI()->cuDeviceGet(&CudaDevice, GPUIndex);
+		if(Status != CUDA_SUCCESS)
+		{
+			UE_LOG(LogRivermax, Warning, TEXT("Can't initialize output to use GPUDirect. Failed to get a Cuda device for GPU %d. Status: %d"), GPUIndex, Status);
+			return false;
+		}
+
+		CUmemAllocationProp AllocationProperties = {};
+		AllocationProperties.type = CU_MEM_ALLOCATION_TYPE_PINNED;
+		AllocationProperties.allocFlags.gpuDirectRDMACapable = 1;
+		AllocationProperties.allocFlags.usage = 0;
+		AllocationProperties.location.type = CU_MEM_LOCATION_TYPE_DEVICE;
+		AllocationProperties.location.id = CudaDevice;
+
+		// Get memory granularity required for cuda device. We need to align allocation with this.
+		size_t Granularity;
+		Status = CudaModule.DriverAPI()->cuMemGetAllocationGranularity(&Granularity, &AllocationProperties, CU_MEM_ALLOC_GRANULARITY_RECOMMENDED);
+		if(Status != CUDA_SUCCESS)
+		{
+			UE_LOG(LogRivermax, Warning, TEXT("Can't initialize output to use GPUDirect. Failed to get allocation granularity. Status: %d"), Status);
+			return false;
+		}
+
+		// Cuda requires allocated memory to be aligned with a certain granularity
+		// We align each frame size to the desired granularity and multiply that by number of buffer
+		// This causes more memory to be allocated but doing a single allocation fails rmax stream creation
+		const size_t CudaAlignedFrameSize = (FrameAllocSize % Granularity) ? FrameAllocSize + (Granularity - (FrameAllocSize % Granularity)) : FrameAllocSize;
+		const size_t TotalCudaAllocSize = CudaAlignedFrameSize * Options.NumberOfBuffers;
+
+		// Reserve contiguous memory to contain required number of buffers. 
+		CUdeviceptr CudaBaseAddress;
+		constexpr CUdeviceptr InitialAddress = 0;
+		constexpr int32 Flags = 0;
+		Status = CudaModule.DriverAPI()->cuMemAddressReserve(&CudaBaseAddress, TotalCudaAllocSize, Granularity, InitialAddress, Flags);
+		if (Status != CUDA_SUCCESS)
+		{
+			UE_LOG(LogRivermax, Warning, TEXT("Can't initialize output to use GPUDirect. Failed to reserve memory for %d bytes. Status: %d"), TotalCudaAllocSize, Status);
+			return false;
+		}
+
+		// Make the allocation on device memory
+		CUmemGenericAllocationHandle Handle;
+		Status = CudaModule.DriverAPI()->cuMemCreate(&Handle, TotalCudaAllocSize, &AllocationProperties, Flags);
+		if (Status != CUDA_SUCCESS)
+		{
+			UE_LOG(LogRivermax, Warning, TEXT("Can't initialize output to use GPUDirect. Failed to create memory on device. Status: %d"), Status);
+			return false;
+		}
+		
+		bool bExit = false;
+		constexpr int32 Offset = 0;
+		Status = CudaModule.DriverAPI()->cuMemMap(CudaBaseAddress, TotalCudaAllocSize, Offset, Handle, Flags);
+		if (Status != CUDA_SUCCESS)
+		{
+			UE_LOG(LogRivermax, Warning, TEXT("Can't initialize output to use GPUDirect. Failed to map memory. Status: %d"), Status);
+			// Need to release handle no matter what
+			bExit = true;
+		}
+		
+		// Cache to know we need to unmap/deallocate even if it fails down the road
+		CudaAllocatedMemory = TotalCudaAllocSize;
+		CudaAllocatedMemoryBaseAddress = reinterpret_cast<void*>(CudaBaseAddress);
+
+		Status = CudaModule.DriverAPI()->cuMemRelease(Handle);
+		if (Status != CUDA_SUCCESS)
+		{
+			UE_LOG(LogRivermax, Warning, TEXT("Can't initialize output to use GPUDirect. Failed to release handle. Status: %d"), Status);
+			return false;
+		}
+
+		if(bExit)
+		{
+			return false;
+		}
+
+		// Setup access description.
+		CUmemAccessDesc MemoryAccessDescription = {};
+		MemoryAccessDescription.flags = CU_MEM_ACCESS_FLAGS_PROT_READWRITE;
+		MemoryAccessDescription.location.type = CU_MEM_LOCATION_TYPE_DEVICE;
+		MemoryAccessDescription.location.id = CudaDevice;
+		constexpr int32 Count = 1;
+		Status = CudaModule.DriverAPI()->cuMemSetAccess(CudaBaseAddress, TotalCudaAllocSize, &MemoryAccessDescription, Count);
+		if (Status != CUDA_SUCCESS)
+		{
+			UE_LOG(LogRivermax, Warning, TEXT("Can't initialize output to use GPUDirect. Failed to configure memory access. Status: %d"), Status);
+			return false;
+		}
+		
+		Status = CudaModule.DriverAPI()->cuCtxSynchronize();
+		if (Status != CUDA_SUCCESS)
+		{
+			UE_LOG(LogRivermax, Warning, TEXT("Can't initialize output to use GPUDirect. Failed to synchronize context. Status: %d"), Status);
+			return false;
+		}
+
+		// Now that Cuda memory was allocated, we can allocate our buffers with an address in cuda space
+		AvailableFrames.Reserve(Options.NumberOfBuffers);
+		for (int32 Index = 0; Index < Options.NumberOfBuffers; ++Index)
+		{
+			TFunction<void(void*)> DeallocatorFunc = nullptr; // We handle a single deallocation for gpu memory
+			TSharedPtr<FRivermaxOutputFrame> Frame = MakeShared<FRivermaxOutputFrame>(Index, DeallocatorFunc);
+			Frame->VideoBuffer = reinterpret_cast<void*>(CudaBaseAddress + (Index * CudaAlignedFrameSize));
+			AvailableFrames.Add(MoveTemp(Frame));
+		}		
+
+		CudaModule.DriverAPI()->cuCtxPopCurrent(nullptr);
+
+		return true;
+	}
+
+	void FRivermaxOutputStream::DeallocateBuffers()
+	{
+		if (CudaAllocatedMemory > 0)
+		{
+			FCUDAModule& CudaModule = FModuleManager::GetModuleChecked<FCUDAModule>("CUDA");
+			CudaModule.DriverAPI()->cuCtxPushCurrent(CudaModule.GetCudaContext());
+
+			const CUdeviceptr BaseAddress = reinterpret_cast<CUdeviceptr>(CudaAllocatedMemoryBaseAddress);
+			CUresult Status = CudaModule.DriverAPI()->cuMemUnmap(BaseAddress, CudaAllocatedMemory);
+			if (Status != CUDA_SUCCESS)
+			{
+				UE_LOG(LogRivermax, Warning, TEXT("Failed to unmap cuda memory. Status: %d"), Status);
+			}
+
+			Status = CudaModule.DriverAPI()->cuMemAddressFree(BaseAddress, CudaAllocatedMemory);
+			if (Status != CUDA_SUCCESS)
+			{
+				UE_LOG(LogRivermax, Warning, TEXT("Failed to free cuda memory. Status: %d"), Status);
+			}
+
+			for (const TPair<FBufferRHIRef, void*>& Entry : BufferCudaMemoryMap)
+			{
+				if(Entry.Value)
+				{
+					CudaModule.DriverAPI()->cuMemFree(reinterpret_cast<CUdeviceptr>(Entry.Value));
+				}
+			}
+			BufferCudaMemoryMap.Empty();
+
+			CudaModule.DriverAPI()->cuCtxPopCurrent(nullptr);
+		}
+	}
+
 	int32 FRivermaxOutputStream::GetStride() const
 	{
 		check(FormatInfo.PixelGroupCoverage != 0);
 		return (Options.AlignedResolution.X / FormatInfo.PixelGroupCoverage) * FormatInfo.PixelGroupSize;
 	}
+
+	void* FRivermaxOutputStream::GetMappedAddress(const FBufferRHIRef& InBuffer)
+	{
+		// If we are here, d3d12 had to have been validated
+		const ERHIInterfaceType RHIType = RHIGetInterfaceType();
+		check(RHIType == ERHIInterfaceType::D3D12);
+
+		//Do we already have a mapped address for this buffer
+		if (BufferCudaMemoryMap.Find((InBuffer)) == nullptr)
+		{
+			int64 BufferMemorySize = 0;
+			CUexternalMemory MappedExternalMemory = nullptr;
+			HANDLE D3D12BufferHandle = 0;
+			CUDA_EXTERNAL_MEMORY_HANDLE_DESC CudaExtMemHandleDesc = {};
+
+			// Create shared handle for our buffer
+			{
+				TRACE_CPUPROFILER_EVENT_SCOPE(Rmax_D3D12CreateSharedHandle);
+
+				ID3D12Resource* NativeD3D12Resource = GetID3D12DynamicRHI()->RHIGetResource(InBuffer);
+				BufferMemorySize = GetID3D12DynamicRHI()->RHIGetResourceMemorySize(InBuffer);
+
+				TRefCountPtr<ID3D12Device> OwnerDevice;
+				HRESULT QueryResult;
+				if ((QueryResult = NativeD3D12Resource->GetDevice(IID_PPV_ARGS(OwnerDevice.GetInitReference()))) != S_OK)
+				{
+					UE_LOG(LogRivermax, Error, TEXT("Failed to get D3D12 device for captured buffer ressource: %d)"), QueryResult);
+					return nullptr;
+				}
+
+				if ((QueryResult = OwnerDevice->CreateSharedHandle(NativeD3D12Resource, NULL, GENERIC_ALL, NULL, &D3D12BufferHandle)) != S_OK)
+				{
+					UE_LOG(LogRivermax, Error, TEXT("Failed to create shared handle for captured buffer ressource: %d"), QueryResult);
+					return nullptr;
+				}
+
+				CudaExtMemHandleDesc.type = CU_EXTERNAL_MEMORY_HANDLE_TYPE_D3D12_RESOURCE;
+				CudaExtMemHandleDesc.handle.win32.name = nullptr;
+				CudaExtMemHandleDesc.handle.win32.handle = D3D12BufferHandle;
+				CudaExtMemHandleDesc.size = BufferMemorySize;
+				CudaExtMemHandleDesc.flags |= CUDA_EXTERNAL_MEMORY_DEDICATED;
+			}
+
+			{
+				TRACE_CPUPROFILER_EVENT_SCOPE(Rmax_CudaImportMemory);
+
+				const CUresult Result = FCUDAModule::CUDA().cuImportExternalMemory(&MappedExternalMemory, &CudaExtMemHandleDesc);
+				
+				if (D3D12BufferHandle)
+				{
+					CloseHandle(D3D12BufferHandle);
+				}
+
+				if (Result != CUDA_SUCCESS)
+				{
+					UE_LOG(LogRivermax, Error, TEXT("Failed to import shared buffer. Error: %d"), Result);
+					return nullptr;
+				}
+			}
+
+			{
+				TRACE_CPUPROFILER_EVENT_SCOPE(Rmax_MapCudaMemory);
+
+				CUDA_EXTERNAL_MEMORY_BUFFER_DESC BufferDescription = {};
+				BufferDescription.offset = 0;
+				BufferDescription.size = BufferMemorySize;
+				CUdeviceptr NewMemory;
+				const CUresult Result = FCUDAModule::CUDA().cuExternalMemoryGetMappedBuffer(&NewMemory, MappedExternalMemory, &BufferDescription);
+				if(Result != CUDA_SUCCESS || NewMemory == 0)
+				{
+					UE_LOG(LogRivermax, Error, TEXT("Failed to get shared buffer mapped memory. Error: %d"), Result);
+					return nullptr;
+				}
+				
+				BufferCudaMemoryMap.Add(InBuffer, reinterpret_cast<void*>(NewMemory));
+			}
+		}
+
+		// At this point, we have the mapped buffer in cuda space and we can use it to schedule a memcpy on cuda engine.
+		return BufferCudaMemoryMap[InBuffer];
+	}
 }
 
diff --git a/Engine/Plugins/VirtualProduction/Rivermax/RivermaxCore/Source/RivermaxCore/Private/RivermaxOutputStream.h b/Engine/Plugins/VirtualProduction/Rivermax/RivermaxCore/Source/RivermaxCore/Private/RivermaxOutputStream.h
--- a/Engine/Plugins/VirtualProduction/Rivermax/RivermaxCore/Source/RivermaxCore/Private/RivermaxOutputStream.h
+++ b/Engine/Plugins/VirtualProduction/Rivermax/RivermaxCore/Source/RivermaxCore/Private/RivermaxOutputStream.h
@@ -4,6 +4,7 @@
 
 #include "IRivermaxOutputStream.h"
 
+#include "Containers/SpscQueue.h"
 #include "HAL/Runnable.h"
 #include "RivermaxHeader.h"
 #include "RivermaxOutputFrame.h"
@@ -126,8 +127,12 @@ namespace UE::RivermaxCore::Private
 		void InitializeStreamTimingSettings();
 		void ShowStats();
 		uint32 GetTimestampFromTime(uint64 InTimeNanosec, double InMediaClockRate) const;
+		
+		bool AllocateGPUBuffers();
 		void AllocateSystemBuffers();
+		void DeallocateBuffers();
 		int32 GetStride() const;
+		void* GetMappedAddress(const FBufferRHIRef& InBuffer);
 
 	private:
 
@@ -187,6 +192,18 @@ namespace UE::RivermaxCore::Private
 		
 		/** Whether stream is using gpudirect to host memory consumed by Rivermax */
 		bool bUseGPUDirect = false;
+
+		/** Allocated memory base address used when it's time to free */
+		void* CudaAllocatedMemoryBaseAddress = nullptr;
+
+		/** Total allocated gpu memory. */
+		int32 CudaAllocatedMemory = 0;
+
+		/** Map between buffer we are sending and their mapped address in gpu space */
+		TMap<FBufferRHIRef, void*> BufferCudaMemoryMap;
+
+		/** Queued identifiers to be consumed by cuda callback when work has been completed */
+		TSpscQueue<uint32> PendingIdentifiers;
 	};
 }
 
diff --git a/Engine/Plugins/VirtualProduction/Rivermax/RivermaxCore/Source/RivermaxCore/Private/RivermaxUtils.h b/Engine/Plugins/VirtualProduction/Rivermax/RivermaxCore/Source/RivermaxCore/Private/RivermaxUtils.h
--- a/Engine/Plugins/VirtualProduction/Rivermax/RivermaxCore/Source/RivermaxCore/Private/RivermaxUtils.h
+++ b/Engine/Plugins/VirtualProduction/Rivermax/RivermaxCore/Source/RivermaxCore/Private/RivermaxUtils.h
@@ -19,7 +19,7 @@ namespace UE::RivermaxCore::Private::Utils
 	static constexpr uint32 FullHDWidth = 1920;
 
 	/** Maximum payload we can send based on UDP max size and RTP header.  */
-	static constexpr uint32 MaxPayloadSize = 1440;
+	static constexpr uint32 MaxPayloadSize = 1420;
 	
 	/** Smallest payload size to bound our search for a payload that can be equal across a line */
 	static constexpr uint32 MinPayloadSize = 950;
diff --git a/Engine/Plugins/VirtualProduction/Rivermax/RivermaxCore/Source/RivermaxCore/Public/RivermaxTypes.h b/Engine/Plugins/VirtualProduction/Rivermax/RivermaxCore/Source/RivermaxCore/Public/RivermaxTypes.h
--- a/Engine/Plugins/VirtualProduction/Rivermax/RivermaxCore/Source/RivermaxCore/Public/RivermaxTypes.h
+++ b/Engine/Plugins/VirtualProduction/Rivermax/RivermaxCore/Source/RivermaxCore/Public/RivermaxTypes.h
@@ -36,7 +36,7 @@ namespace UE::RivermaxCore
 		FIntPoint AlignedResolution = FIntPoint::ZeroValue;
 
 		/** Whether to leverage GPUDirect (Cuda) capability to transfer memory to NIC if available */
-		bool bUseGPUDirect = false;
+		bool bUseGPUDirect = true;
 	};
 
 	enum class RIVERMAXCORE_API ERivermaxStreamType : uint8
diff --git a/Engine/Plugins/VirtualProduction/Rivermax/RivermaxCore/Source/RivermaxCore/RivermaxCore.build.cs b/Engine/Plugins/VirtualProduction/Rivermax/RivermaxCore/Source/RivermaxCore/RivermaxCore.build.cs
--- a/Engine/Plugins/VirtualProduction/Rivermax/RivermaxCore/Source/RivermaxCore/RivermaxCore.build.cs
+++ b/Engine/Plugins/VirtualProduction/Rivermax/RivermaxCore/Source/RivermaxCore/RivermaxCore.build.cs
@@ -13,17 +13,18 @@ public class RivermaxCore : ModuleRules
 				"DeveloperSettings",
 				"MediaAssets",
 				"MediaIOCore",
-				"RivermaxLib",
 				"RHI",
+				"RivermaxLib",
 				"TimeManagement"
 			}
 		);
 
 		PrivateDependencyModuleNames.AddRange(
 			new string[] {
+				"D3D12RHI",
 				"Engine",
 				"Networking",
-				"RenderCore",
+				"RenderCore"
 			}
 		);
 
@@ -39,5 +40,11 @@ public class RivermaxCore : ModuleRules
 				}
 			);
 		}
+
+		if (Target.Platform.IsInGroup(UnrealPlatformGroup.Windows))
+		{
+			AddEngineThirdPartyPrivateStaticDependencies(Target, "CUDA", "DX12");
+		}
+
 	}
 }
diff --git a/Engine/Plugins/VirtualProduction/Rivermax/RivermaxMedia/Source/RivermaxMedia/Public/RivermaxMediaOutput.h b/Engine/Plugins/VirtualProduction/Rivermax/RivermaxMedia/Source/RivermaxMedia/Public/RivermaxMediaOutput.h
--- a/Engine/Plugins/VirtualProduction/Rivermax/RivermaxMedia/Source/RivermaxMedia/Public/RivermaxMediaOutput.h
+++ b/Engine/Plugins/VirtualProduction/Rivermax/RivermaxMedia/Source/RivermaxMedia/Public/RivermaxMediaOutput.h
@@ -77,5 +77,5 @@ public:
 
 	/** Whether to use GPUDirect if available (Memcopy from GPU to NIC directly bypassing system) if available */
 	UPROPERTY(EditAnywhere, AdvancedDisplay, Category = "Output")
-	bool bUseGPUDirect = false;
+	bool bUseGPUDirect = true;
 };
