- Adding Rivermax gpudirect support for input
   * Added shared buffer cuda mapping support in input stream
   * Added gpu buffer path to rivermax samples
- Updated some of the trace markups to track frames
- Added some utils for 2110-20 formats definition
- Updated our RTP representation but still need some validation with other hardware
- Changed cuda work (memcpy asyncs) used for input and output streams to be done on separate cuda streams

diff --git a/Engine/Plugins/VirtualProduction/Rivermax/RivermaxCore/Source/RivermaxCore/Private/RivermaxInputStream.cpp b/Engine/Plugins/VirtualProduction/Rivermax/RivermaxCore/Source/RivermaxCore/Private/RivermaxInputStream.cpp
--- a/Engine/Plugins/VirtualProduction/Rivermax/RivermaxCore/Source/RivermaxCore/Private/RivermaxInputStream.cpp
+++ b/Engine/Plugins/VirtualProduction/Rivermax/RivermaxCore/Source/RivermaxCore/Private/RivermaxInputStream.cpp
@@ -3,6 +3,8 @@
 #include "RivermaxInputStream.h"
 
 #include "Async/Async.h"
+#include "CudaModule.h"
+#include "ID3D12DynamicRHI.h"
 #include "IRivermaxCoreModule.h"
 #include "IRivermaxManager.h"
 #include "Misc/ByteSwap.h"
@@ -15,17 +17,23 @@
 
 namespace UE::RivermaxCore::Private
 {
-	static TAutoConsoleVariable<int32> CVarUseIncompleteFrames(
-		TEXT("Rivermax.Input.UseIncompleteFrames"),
-		1,
-		TEXT("Whether to keep incomplete frames in the case of missing packets or drop them."),
+	static TAutoConsoleVariable<float> CVarWaitForCompletionTimeout(
+		TEXT("Rivermax.Input.WaitCompletionTimeout"),
+		0.25,
+		TEXT("Maximum time to wait, in seconds, when waiting for a memory copy operation to complete on the gpu."),
+		ECVF_Default);
+
+	static TAutoConsoleVariable<int32> CVarExpectedPayloadSize(
+		TEXT("Rivermax.Input.ExpectedPayloadSize"),
+		1500,
+		TEXT("Expected payload size used to initialize rivermax stream."),
 		ECVF_Default);
 
 
 #if PLATFORM_SUPPORTS_PRAGMA_PACK
 #pragma pack(push, 1)
 #endif
-	struct FRivermaxRTPSampleRowData
+	struct FRawRTPHeader
 	{
 		uint32 ContributingSourceCount : 4;
 		uint32 ExtensionBit : 1;
@@ -54,23 +62,59 @@ namespace UE::RivermaxCore::Private
 		uint32 ContinuationBit2 : 1;
 		uint32 SRDOffsetLow2 : 8;
 
-		uint16 GetSrd1RowNumber() { return ((SRDRowNumberHigh1 << 8) | SRDRowNumberLow1); }
-		uint16 GetSrd1Offset() { return ((SRDOffsetHigh1 << 8) | SRDOffsetLow1); }
+		uint16 GetSrd1RowNumber() const { return ((SRDRowNumberHigh1 << 8) | SRDRowNumberLow1); }
+		uint16 GetSrd1Offset() const { return ((SRDOffsetHigh1 << 8) | SRDOffsetLow1); }
 
-		uint16 GetSrd2RowNumber() { return ((SRDRowNumberHigh2 << 8) | SRDRowNumberLow2); }
-		uint16 GetSrd2Offset() { return ((SRDOffsetHigh2 << 8) | SRDOffsetLow2); }
+		uint16 GetSrd2RowNumber() const { return ((SRDRowNumberHigh2 << 8) | SRDRowNumberLow2); }
+		uint16 GetSrd2Offset() const { return ((SRDOffsetHigh2 << 8) | SRDOffsetLow2); }
 	};
 #if PLATFORM_SUPPORTS_PRAGMA_PACK
 #pragma pack(pop)
 #endif
 
+	struct FSRDHeader
+	{
+		/** Length of payload. Is a multiple of pgroup (see pixel formats) */
+		uint16 Length = 0;
+
+		/** False if progressive or first field of interlace. True if second field of interlace */
+		bool bIsFieldOne = false;
+
+		/** Video line number, starts at 0 */
+		uint16 RowNumber = 0;
+
+		/** Whether another SRD is following this one */
+		bool bHasContinuation = false;
+		
+		/** Location of the first pixel in payload, in pixel */
+		uint16 DataOffset = 0;
+	};
+	
+	struct FRTPHeader
+	{
+		/** Sequence number including extension if present */
+		uint32 SequencerNumber = 0;
+
+		/** Timestamp of frame in the specified clock resolution. Video is typically 90kHz */
+		uint32 Timestamp = 0;
+
+		/** Identification of this stream */
+		uint32 SyncSouceId = 0;
+
+		/** Whether extensions (SRD headers) are present */
+		bool bHasExtension = false;
+
+		/** True if RTP packet is last of video stream */
+		bool bIsMarkerBit = false;
+
+		/** Only supports 2 SRD for now. Adjust if needed */
+		FSRDHeader SRD1;
+		FSRDHeader SRD2;
+	};
 
 	uint8* GetRTPHeaderPointer(uint8* InHeader)
 	{
-		if (!InHeader /*|| RMAX_APP_PROTOCOL_PACKET == m_type*/)
-		{
-			return nullptr;
-		}
+		check(InHeader);
 
 		static constexpr uint32 ETH_TYPE_802_1Q = 0x8100;          /* 802.1Q VLAN Extended Header  */
 		static constexpr uint32 RTP_HEADER_SIZE = 12;
@@ -108,13 +152,14 @@ namespace UE::RivermaxCore::Private
 		Options = InOptions;
 		Listener = &InListener;
 		FormatInfo = FStandardVideoFormat::GetVideoFormatInfo(Options.PixelFormat);
+		ExpectedPayloadSize = CVarExpectedPayloadSize.GetValueOnGameThread();
 
 		InitTaskFuture = Async(EAsyncExecution::TaskGraph, [this]()
 		{
 			// If the stream is trying to shutdown before the init task has even started, don't bother
 			if (bIsShuttingDown)
 			{
-					return;
+				return;
 			}
 
 			bool bWasSuccessful = false;
@@ -153,11 +198,8 @@ namespace UE::RivermaxCore::Private
 
 					FMemory::Memset(&BufferConfiguration.DataMemory, 0, sizeof(BufferConfiguration.DataMemory));
 
-					// Todo : Unreal doesn't send Payloads greater than 1280 but other devices could.
-					// While consuming input chunks, we memcpy based on SRD data but SRD data could be greater than 
-					// max_size so our loop should be updated to support that case.
-					BufferConfiguration.DataMemory.min_size = 1500; 
-					BufferConfiguration.DataMemory.max_size = 1500;
+					BufferConfiguration.DataMemory.min_size = ExpectedPayloadSize;
+					BufferConfiguration.DataMemory.max_size = ExpectedPayloadSize;
 					BufferAttributes.data = &BufferConfiguration.DataMemory;
 
 					FMemory::Memset(&BufferConfiguration.HeaderMemory, 0, sizeof(BufferConfiguration.HeaderMemory));
@@ -167,10 +209,7 @@ namespace UE::RivermaxCore::Private
 					rmax_status_t Status = rmax_in_query_buffer_size(StreamType, &RivermaxInterface, &BufferAttributes, &BufferConfiguration.PayloadSize, &BufferConfiguration.HeaderSize);
 					if (Status == RMAX_OK)
 					{
-						const uint32 CacheLineSize = PLATFORM_CACHE_LINE_SIZE;
-
-						BufferConfiguration.DataMemory.ptr = FMemory::Malloc(BufferConfiguration.PayloadSize, CacheLineSize);
-						BufferConfiguration.HeaderMemory.ptr = FMemory::Malloc(BufferConfiguration.HeaderSize, CacheLineSize);
+						AllocateBuffers();
 
 						//Buffers configured, now configure stream and attach flow
 						const rmax_in_timestamp_format TimestampFormat = rmax_in_timestamp_format::RMAX_PACKET_TIMESTAMP_RAW_NANO; //how packets are stamped. counter or nanoseconds
@@ -184,6 +223,11 @@ namespace UE::RivermaxCore::Private
 								bIsActive = true;
 								RivermaxThread = FRunnableThread::Create(this, TEXT("Rivermax InputStream Thread"), 128 * 1024, TPri_AboveNormal, FPlatformAffinity::GetPoolThreadMask());
 								bWasSuccessful = true;
+
+								UE_LOG(LogRivermax, Display, TEXT("Input started receiving %dx%d%s")
+									, Options.AlignedResolution.X
+									, Options.AlignedResolution.Y
+									, bIsUsingGPUDirect ? TEXT(" using GPUDirect") : TEXT(""));
 							}
 							else
 							{
@@ -226,6 +270,8 @@ namespace UE::RivermaxCore::Private
 			RivermaxThread = nullptr;
 			UE_LOG(LogRivermax, Log, TEXT("Rivermax Input stream has shutdown"));
 		}
+
+		DeallocateBuffers();
 	}
 
 	void FRivermaxInputStream::Process_AnyThread()
@@ -239,7 +285,7 @@ namespace UE::RivermaxCore::Private
 		rmax_status_t Status = rmax_in_get_next_chunk(StreamId, MinChunkSize, MaxChunkSize, Timeout, Flags, &Completion);
 		if (Status == RMAX_OK)
 		{
-			ParseChunk(Completion);
+			ParseChunks(Completion);
 		}
 		else
 		{
@@ -289,30 +335,47 @@ namespace UE::RivermaxCore::Private
 
 	}
 
-	bool FRivermaxInputStream::GetRTPParameter(uint8* InputRTP, FRTPParameter& OutParameter)
+	bool FRivermaxInputStream::TranslateRTPHeader(const FRawRTPHeader& RawRTPHeader, FRTPHeader& OutRTPHeader)
 	{
-		OutParameter.Timestamp = 0;
+		OutRTPHeader.Timestamp = 0;
 		
-		const uint8* RTPHeaderPtr = GetRTPHeaderPointer(InputRTP);
-		if (((RTPHeaderPtr[0] & 0xC0) != 0x80))
+		if (RawRTPHeader.Version != 2)
 		{
 			return false;
 		}
 
-		OutParameter.SequencerNumber = RTPHeaderPtr[3] | RTPHeaderPtr[2] << 8;
-		if (IsExtendedSequenceNumber()) //todo if 2022 is supported extended sequence number is supported in other ways
+		// Pretty sure some data needs to be swapped but can't validate that until we have other hardware generating data
+		OutRTPHeader.SequencerNumber = (ByteSwap((uint16)RawRTPHeader.ExtendedSequenceNumber) << 16) | ByteSwap((uint16)RawRTPHeader.SequenceNumber);
+		OutRTPHeader.Timestamp = ByteSwap(RawRTPHeader.Timestamp);
+		OutRTPHeader.bIsMarkerBit = RawRTPHeader.MarkerBit;
+
+		OutRTPHeader.SyncSouceId = RawRTPHeader.SynchronizationSource;
+
+		OutRTPHeader.SRD1.Length = ByteSwap((uint16)RawRTPHeader.SRDLength1);
+		OutRTPHeader.SRD1.DataOffset = RawRTPHeader.GetSrd1Offset();
+		OutRTPHeader.SRD1.RowNumber = RawRTPHeader.GetSrd1RowNumber();
+		OutRTPHeader.SRD1.bIsFieldOne = RawRTPHeader.FieldIdentification1;
+		OutRTPHeader.SRD1.bHasContinuation = RawRTPHeader.ContinuationBit1;
+
+		if (OutRTPHeader.SRD1.bHasContinuation)
 		{
-			OutParameter.SequencerNumber |= RTPHeaderPtr[12] << 24 | RTPHeaderPtr[13] << 16;
-			OutParameter.bIsFBit = !!(RTPHeaderPtr[16] & 0x80);
+			OutRTPHeader.SRD2.Length = ByteSwap((uint16)RawRTPHeader.SRDLength2);
+			OutRTPHeader.SRD2.DataOffset = RawRTPHeader.GetSrd2Offset();
+			OutRTPHeader.SRD2.RowNumber = RawRTPHeader.GetSrd2RowNumber();
+			OutRTPHeader.SRD2.bIsFieldOne = RawRTPHeader.FieldIdentification2;
+			OutRTPHeader.SRD2.bHasContinuation = RawRTPHeader.ContinuationBit2;
+
+			if (OutRTPHeader.SRD2.bHasContinuation == true)
+			{
+				UE_LOG(LogRivermax, Verbose, TEXT("Received SRD with more than 2 SRD which isn't supported."));
+			}
 		}
-		OutParameter.Timestamp = ByteSwap(*(uint32*)((RTPHeaderPtr)+4));
-		OutParameter.bIsMBit = !!(RTPHeaderPtr[1] & 0x80);
+
 		return true;
 	}
 
-	void FRivermaxInputStream::ParseChunk(const rmax_in_completion& Completion)
+	void FRivermaxInputStream::ParseChunks(const rmax_in_completion& Completion)
 	{
-		
 		for (uint64 StrideIndex = 0; StrideIndex < Completion.chunk_size; ++StrideIndex)
 		{
 			++StreamStats.ChunksReceived;
@@ -323,40 +386,49 @@ namespace UE::RivermaxCore::Private
 				break;
 			}
 
-			uint8* HeaderPtr = (uint8*)Completion.hdr_ptr + StrideIndex * (size_t)BufferConfiguration.HeaderMemory.stride_size; //when using RMAX_RAW_PACKET then header is preceded by net header
+			uint8* RawHeaderPtr = (uint8*)Completion.hdr_ptr + StrideIndex * (size_t)BufferConfiguration.HeaderMemory.stride_size; 
 			uint8* DataPtr = (uint8*)Completion.data_ptr + StrideIndex * (size_t)BufferConfiguration.DataMemory.stride_size; // The payload is our data
-			uint32 DataOffset = 0;
 
-			if (Completion.packet_info_arr[StrideIndex].data_size)
+			if (Completion.packet_info_arr[StrideIndex].data_size && RawHeaderPtr && DataPtr)
 			{
-				FRTPParameter Parameter;
-				const bool bIsValid = GetRTPParameter(HeaderPtr, Parameter);
+				FRTPHeader RTPHeader;
+
+				// Get RTPHeader address from the raw net header
+				const FRawRTPHeader& RawRTPHeaderPtr = reinterpret_cast<const FRawRTPHeader&>(*GetRTPHeaderPointer(RawHeaderPtr));
+				const bool bIsValid = TranslateRTPHeader(RawRTPHeaderPtr, RTPHeader);
 				if (bIsValid)
 				{
+					// Add trace for the first packet of a frame to help visualize reception of a full frame in time
+					if (bIsFirstPacketReceived == false)
+					{
+						const FString TraceName = FString::Format(TEXT("RmaxInput::StartingFrame {0}"), { RTPHeader.Timestamp });
+						TRACE_CPUPROFILER_EVENT_SCOPE_TEXT(*TraceName);
+						bIsFirstPacketReceived = true;
+					}
+
 					if (bIsFirstFrameReceived)
 					{
 						StreamStats.BytesReceived += Completion.packet_info_arr[StrideIndex].data_size + Completion.packet_info_arr[StrideIndex].hdr_size;
 
 						uint64 LastSequenceNumberIncremented = StreamData.LastSequenceNumber + 1;
-						if (IsExtendedSequenceNumber() == false)
-						{
-							LastSequenceNumberIncremented = LastSequenceNumberIncremented & 0xFFFF;
-						}
 
-						bool bCanProcessChunk = true;
-						const uint64 LostPackets = ((uint64)Parameter.SequencerNumber + 0x100000000 - LastSequenceNumberIncremented) & 0xFFFFFFFF;
+						bool bCanProcessSRD = true;
+						const uint64 LostPackets = ((uint64)RTPHeader.SequencerNumber + 0x100000000 - LastSequenceNumberIncremented) & 0xFFFFFFFF;
 						if (LostPackets > 0)
 						{
-							bCanProcessChunk = false;
+							bCanProcessSRD = false;
 							StreamData.WritingOffset = 0;
 							StreamData.ReceivedSize = 0;
 							++StreamStats.TotalPacketLossCount;
 							++StreamStats.FramePacketLossCount;
 							
+							// For now, if packets were lost, skip the incoming frame. We could improve that and have corrupted frames instead of skipping them but can be added later
+							bIsFirstFrameReceived = false;
+							
 							UE_LOG(LogRivermax, Warning, TEXT("Lost %llu packets"), LostPackets);
 						}
 
-						StreamData.LastSequenceNumber = Parameter.SequencerNumber;
+						StreamData.LastSequenceNumber = RTPHeader.SequencerNumber;
 
 						//if flags are RMAX_IN_CREATE_STREAM_INFO_PER_PACKET todo 
 						{
@@ -366,25 +438,10 @@ namespace UE::RivermaxCore::Private
 							}
 						}
 
-						if (bCanProcessChunk)
+						if (bCanProcessSRD)
 						{
-							const FRivermaxRTPSampleRowData* HeaderStart = reinterpret_cast<FRivermaxRTPSampleRowData*>(GetRTPHeaderPointer(HeaderPtr));
-
-							uint16 SizeToCopy = ByteSwap((uint16)HeaderStart->SRDLength1);
-							FMemory::Memcpy(&StreamData.CurrentFrame[StreamData.WritingOffset], &DataPtr[DataOffset], SizeToCopy);
-							StreamData.WritingOffset += SizeToCopy;
-							StreamData.ReceivedSize += SizeToCopy;
-
-							// todo Warning !!! GPUDirect doesn't support more than 1 SRD
-							if (HeaderStart->ContinuationBit1)
-							{
-								DataOffset += SizeToCopy;
-								SizeToCopy = ByteSwap((uint16)HeaderStart->SRDLength2);
-								FMemory::Memcpy(&StreamData.CurrentFrame[StreamData.WritingOffset], &DataPtr[DataOffset], SizeToCopy);
-								StreamData.WritingOffset += SizeToCopy;
-								StreamData.ReceivedSize += SizeToCopy;
-							}
-
+							ProcessSRD(RTPHeader, DataPtr);
+							
 							if (StreamData.ReceivedSize > StreamData.ExpectedSize)
 							{
 								UE_LOG(LogRivermax, Warning, TEXT("Received too much data (%d). Expected %d but received (%d)"), StreamData.ReceivedSize - StreamData.ExpectedSize, StreamData.ExpectedSize, StreamData.ReceivedSize);
@@ -392,50 +449,24 @@ namespace UE::RivermaxCore::Private
 								StreamData.ReceivedSize = 0;
 								++StreamStats.BiggerFramesCount;
 							}
-							else if (HeaderStart->MarkerBit)
+							else if (RTPHeader.bIsMarkerBit)
 							{
-								if (StreamData.ReceivedSize == StreamData.ExpectedSize)
-								{
-									TRACE_BOOKMARK(TEXT("RMAX RX %llu"), Parameter.Timestamp);
-									TRACE_CPUPROFILER_EVENT_SCOPE(FRivermaxInputStream::ProcessingReceivedFrame)
-
-									++StreamStats.FramesReceived;
-
-									FRivermaxInputVideoFrameDescriptor Descriptor;
-									Descriptor.Width = Options.Resolution.X;
-									Descriptor.Height = Options.Resolution.Y;
-									Descriptor.Stride = Options.AlignedResolution.X / FormatInfo.PixelGroupCoverage * FormatInfo.PixelGroupSize;
-									FRivermaxInputVideoFrameReception NewFrame;
-									NewFrame.VideoBuffer = StreamData.CurrentFrame;
-									Listener->OnVideoFrameReceived(Descriptor, NewFrame);
-									PrepareNextFrame();
-								}
-								else
-								{
-									UE_LOG(LogRivermax, Warning, TEXT("End of frame received (Marker bit) but not enough data was received (missing %d). Expected %d but received (%d)"), StreamData.ExpectedSize - StreamData.ReceivedSize, StreamData.ExpectedSize, StreamData.ReceivedSize);
-									StreamData.WritingOffset = 0;
-									StreamData.ReceivedSize = 0;
-									++StreamStats.InvalidFramesCount;
-								}
+								ProcessLastSRD(RTPHeader, DataPtr);
 								
 								StreamStats.FramePacketLossCount = 0;
 								++StreamStats.EndOfFrameReceived;
 							}
 						}
-
 					}
 					else
 					{
-						if (Parameter.bIsMBit)
+						if (RTPHeader.bIsMarkerBit)
 						{
-							StreamData.LastSequenceNumber = Parameter.SequencerNumber;
-							bIsFirstFrameReceived = true;
+							StreamData.LastSequenceNumber = RTPHeader.SequencerNumber;
 							PrepareNextFrame();
+							bIsFirstFrameReceived = (StreamData.CurrentFrame != nullptr);
 						}
 					}
-
-					//todo why  Reset RTP header
-					*(uint32_t*)GetRTPHeaderPointer(HeaderPtr) = 0;
 				}
 				else
 				{
@@ -446,30 +477,55 @@ namespace UE::RivermaxCore::Private
 			{
 				++StreamStats.EmptyCompletionCount;
 			}
-
-			//todo why?
-			Completion.packet_info_arr[StrideIndex].data_size = 0;
 		}
 	}
 
-	bool FRivermaxInputStream::IsExtendedSequenceNumber() const
-	{
-		return RivermaxStreamType == ERivermaxStreamType::VIDEO_2110_20_STREAM;
-	}
-
 	void FRivermaxInputStream::PrepareNextFrame()
 	{
 		using namespace UE::RivermaxCore::Private::Utils;
 
 		FRivermaxInputVideoFrameDescriptor Descriptor;
+		Descriptor.bIsUsingGPUDirect = bIsUsingGPUDirect;
 		FRivermaxInputVideoFrameRequest Request;
 		const int32 Stride = Options.AlignedResolution.X / FormatInfo.PixelGroupCoverage * FormatInfo.PixelGroupSize;
 		Descriptor.VideoBufferSize = Options.Resolution.Y * Stride;
 		Listener->OnVideoFrameRequested(Descriptor, Request);
-		StreamData.CurrentFrame = Request.VideoBuffer;
+
+		// Reset current frame to know when we have a valid one
+		StreamData.CurrentFrame = nullptr;
+		if (bIsUsingGPUDirect)
+		{
+			if(Request.GPUBuffer)
+			{
+				StreamData.CurrentFrame = GetMappedBuffer(Request.GPUBuffer);
+			}
+		}
+		else
+		{	
+			if (Request.VideoBuffer)
+			{
+				StreamData.CurrentFrame = Request.VideoBuffer;
+			}
+		}
+
+		// Verify if we were able to request a valid frame. If engine is blocked, it could happen that there is none available 
+		if (StreamData.CurrentFrame == nullptr)
+		{
+			// If we failed getting one, reset the valid first frame received and wait for the next one
+			bIsFirstFrameReceived = false;
+		}
+
 		StreamData.WritingOffset = 0;
 		StreamData.ReceivedSize = 0;
 		StreamData.ExpectedSize = Descriptor.VideoBufferSize;
+		StreamData.DeviceWritePointerOne = nullptr;
+		StreamData.SizeToWriteOne = 0;
+		StreamData.DeviceWritePointerTwo = nullptr;
+		StreamData.SizeToWriteTwo = 0;
+		bIsFirstPacketReceived = false;
+
+		// New frame starting, reset tracked SRD
+		LastSRDLength.Reset();
 	}
 
 	void FRivermaxInputStream::LogStats()
@@ -496,5 +552,456 @@ namespace UE::RivermaxCore::Private
 		}
 	}
 
+	void FRivermaxInputStream::AllocateBuffers()
+	{
+		IRivermaxCoreModule& RivermaxModule = FModuleManager::LoadModuleChecked<IRivermaxCoreModule>(TEXT("RivermaxCore"));
+		if (RivermaxModule.GetRivermaxManager()->IsGPUDirectSupported() && Options.bUseGPUDirect)
+		{
+			bIsUsingGPUDirect = AllocateGPUBuffers();
+		}
+
+		constexpr uint32 CacheLineSize = PLATFORM_CACHE_LINE_SIZE;
+		if (bIsUsingGPUDirect == false)
+		{
+			BufferConfiguration.DataMemory.ptr = FMemory::Malloc(BufferConfiguration.PayloadSize, CacheLineSize);
+		}
+		
+		BufferConfiguration.HeaderMemory.ptr = FMemory::Malloc(BufferConfiguration.HeaderSize, CacheLineSize);
+	}
+
+	bool FRivermaxInputStream::AllocateGPUBuffers()
+	{
+		// Allocate memory space where rivermax input will write received buffer to
+
+		TRACE_CPUPROFILER_EVENT_SCOPE(FRivermaxInputStream::AllocateGPUBuffers);
+
+		const ERHIInterfaceType RHIType = RHIGetInterfaceType();
+		if (RHIType != ERHIInterfaceType::D3D12)
+		{
+			UE_LOG(LogRivermax, Warning, TEXT("Can't initialize input to use GPUDirect. RHI is %d but only Dx12 is supported at the moment."), RHIType);
+			return false;
+		}
+
+		FCUDAModule& CudaModule = FModuleManager::GetModuleChecked<FCUDAModule>("CUDA");
+
+		CudaModule.DriverAPI()->cuCtxPushCurrent(CudaModule.GetCudaContextForDevice(DeviceIndex));
+
+		// Todo: Add support for mgpu. 
+		CUdevice CudaDevice;
+		CUresult Status = CudaModule.DriverAPI()->cuDeviceGet(&CudaDevice, DeviceIndex);
+		if (Status != CUDA_SUCCESS)
+		{
+			UE_LOG(LogRivermax, Warning, TEXT("Can't initialize input to use GPUDirect. Failed to get a Cuda device for GPU %d. Status: %d"), DeviceIndex, Status);
+			return false;
+		}
+
+		CUmemAllocationProp AllocationProperties = {};
+		AllocationProperties.type = CU_MEM_ALLOCATION_TYPE_PINNED;
+		AllocationProperties.allocFlags.gpuDirectRDMACapable = 1; //is that required?
+		AllocationProperties.allocFlags.usage = 0;
+		AllocationProperties.location.type = CU_MEM_LOCATION_TYPE_DEVICE;
+		AllocationProperties.location.id = CudaDevice;
+
+		// Get memory granularity required for cuda device. We need to align allocation with this.
+		size_t Granularity;
+		Status = CudaModule.DriverAPI()->cuMemGetAllocationGranularity(&Granularity, &AllocationProperties, CU_MEM_ALLOC_GRANULARITY_RECOMMENDED);
+		if (Status != CUDA_SUCCESS)
+		{
+			UE_LOG(LogRivermax, Warning, TEXT("Can't initialize input to use GPUDirect. Failed to get allocation granularity. Status: %d"), Status);
+			return false;
+		}
+
+		// Cuda requires allocated memory to be aligned with a certain granularity
+		const size_t CudaAlignedAllocation = (BufferConfiguration.PayloadSize % Granularity) ? BufferConfiguration.PayloadSize + (Granularity - (BufferConfiguration.PayloadSize % Granularity)) : BufferConfiguration.PayloadSize;
+		
+		CUdeviceptr CudaBaseAddress;
+		constexpr CUdeviceptr InitialAddress = 0;
+		constexpr int32 Flags = 0;
+		Status = CudaModule.DriverAPI()->cuMemAddressReserve(&CudaBaseAddress, CudaAlignedAllocation, Granularity, InitialAddress, Flags);
+		if (Status != CUDA_SUCCESS)
+		{
+			UE_LOG(LogRivermax, Warning, TEXT("Can't initialize input to use GPUDirect. Failed to reserve memory for %d bytes. Status: %d"), CudaAlignedAllocation, Status);
+			return false;
+		}
+
+		// Make the allocation on device memory
+		CUmemGenericAllocationHandle Handle;
+		Status = CudaModule.DriverAPI()->cuMemCreate(&Handle, CudaAlignedAllocation, &AllocationProperties, Flags);
+		if (Status != CUDA_SUCCESS)
+		{
+			UE_LOG(LogRivermax, Warning, TEXT("Can't initialize input to use GPUDirect. Failed to create memory on device. Status: %d"), Status);
+			return false;
+		}
+		UE_LOG(LogRivermax, Verbose, TEXT("Allocated %d cuda memory"), CudaAlignedAllocation);
+
+		bool bExit = false;
+		constexpr int32 Offset = 0;
+		Status = CudaModule.DriverAPI()->cuMemMap(CudaBaseAddress, CudaAlignedAllocation, Offset, Handle, Flags);
+		if (Status != CUDA_SUCCESS)
+		{
+			UE_LOG(LogRivermax, Warning, TEXT("Can't initialize input to use GPUDirect. Failed to map memory. Status: %d"), Status);
+			// Need to release handle no matter what
+			bExit = true;
+		}
+
+		GPUAllocatedMemorySize = CudaAlignedAllocation;
+		GPUAllocatedMemoryBaseAddress = reinterpret_cast<void*>(CudaBaseAddress);
+
+		Status = CudaModule.DriverAPI()->cuMemRelease(Handle);
+		if (Status != CUDA_SUCCESS)
+		{
+			UE_LOG(LogRivermax, Warning, TEXT("Can't initialize input to use GPUDirect. Failed to release handle. Status: %d"), Status);
+			return false;
+		}
+
+		if (bExit)
+		{
+			return false;
+		}
+
+		// Setup access description.
+		CUmemAccessDesc MemoryAccessDescription = {};
+		MemoryAccessDescription.flags = CU_MEM_ACCESS_FLAGS_PROT_READWRITE;
+		MemoryAccessDescription.location.type = CU_MEM_LOCATION_TYPE_DEVICE;
+		MemoryAccessDescription.location.id = CudaDevice;
+		constexpr int32 Count = 1;
+		Status = CudaModule.DriverAPI()->cuMemSetAccess(CudaBaseAddress, CudaAlignedAllocation, &MemoryAccessDescription, Count);
+		if (Status != CUDA_SUCCESS)
+		{
+			UE_LOG(LogRivermax, Warning, TEXT("Can't initialize input to use GPUDirect. Failed to configure memory access. Status: %d"), Status);
+			return false;
+		}
+
+		CUstream CudaStream;
+		Status = CudaModule.DriverAPI()->cuStreamCreate(&CudaStream, CU_STREAM_NON_BLOCKING);
+		if (Status != CUDA_SUCCESS)
+		{
+			UE_LOG(LogRivermax, Warning, TEXT("Can't initialize input to use GPUDirect. Failed to create its stream. Status: %d"), Status);
+			return false;
+		}
+
+		GPUStream = CudaStream;
+
+		Status = CudaModule.DriverAPI()->cuCtxSynchronize();
+		if (Status != CUDA_SUCCESS)
+		{
+			UE_LOG(LogRivermax, Warning, TEXT("Can't initialize input to use GPUDirect. Failed to synchronize context. Status: %d"), Status);
+			return false;
+		}
+
+		// Give rivermax input buffer config the pointer to gpu allocated memory
+		BufferConfiguration.DataMemory.ptr = GPUAllocatedMemoryBaseAddress;
+
+		CallbackPayload = MakeShared<FCallbackPayload>();
+
+		return true;
+	}
+
+	void FRivermaxInputStream::DeallocateBuffers()
+	{
+		if (GPUAllocatedMemorySize > 0)
+		{
+			FCUDAModule& CudaModule = FModuleManager::GetModuleChecked<FCUDAModule>("CUDA");
+			CudaModule.DriverAPI()->cuCtxPushCurrent(CudaModule.GetCudaContextForDevice(DeviceIndex));
+
+			const CUdeviceptr BaseAddress = reinterpret_cast<CUdeviceptr>(GPUAllocatedMemoryBaseAddress);
+			CUresult Status = CudaModule.DriverAPI()->cuMemUnmap(BaseAddress, GPUAllocatedMemorySize);
+			if (Status != CUDA_SUCCESS)
+			{
+				UE_LOG(LogRivermax, Warning, TEXT("Failed to unmap cuda memory used for input stream. Status: %d"), Status);
+			}
+
+			Status = CudaModule.DriverAPI()->cuMemAddressFree(BaseAddress, GPUAllocatedMemorySize);
+			if (Status != CUDA_SUCCESS)
+			{
+				UE_LOG(LogRivermax, Warning, TEXT("Failed to free cuda memory used for input stream. Status: %d"), Status);
+			}
+			UE_LOG(LogRivermax, Verbose, TEXT("Deallocated %d cuda memory at address %d"), GPUAllocatedMemorySize, GPUAllocatedMemoryBaseAddress);
+
+			GPUAllocatedMemorySize = 0;
+			GPUAllocatedMemoryBaseAddress = 0;
+
+
+			for (const TPair<FRHIBuffer*, void*>& Entry : BufferGPUMemoryMap)
+			{
+				if (Entry.Value)
+				{
+					CudaModule.DriverAPI()->cuMemFree(reinterpret_cast<CUdeviceptr>(Entry.Value));
+				}
+			}
+			BufferGPUMemoryMap.Empty();
+
+			Status = CudaModule.DriverAPI()->cuStreamDestroy(reinterpret_cast<CUstream>(GPUStream));
+			if (Status != CUDA_SUCCESS)
+			{
+				UE_LOG(LogRivermax, Warning, TEXT("Failed to destroy cuda stream. Status: %d"), Status);
+			}
+			GPUStream = nullptr;
+
+
+			CudaModule.DriverAPI()->cuCtxPopCurrent(nullptr);
+		}
+	}
+
+	void* FRivermaxInputStream::GetMappedBuffer(const FBufferRHIRef& InBuffer)
+	{
+		// If we are here, d3d12 had to have been validated
+		const ERHIInterfaceType RHIType = RHIGetInterfaceType();
+		check(RHIType == ERHIInterfaceType::D3D12);
+
+		//Do we already have a mapped address for this buffer
+		if (BufferGPUMemoryMap.Find((InBuffer)) == nullptr)
+		{
+			int64 BufferMemorySize = 0;
+			CUexternalMemory MappedExternalMemory = nullptr;
+			HANDLE D3D12BufferHandle = 0;
+			CUDA_EXTERNAL_MEMORY_HANDLE_DESC CudaExtMemHandleDesc = {};
+
+			// Create shared handle for our buffer
+			{
+				TRACE_CPUPROFILER_EVENT_SCOPE(RmaxInput_D3D12CreateSharedHandle);
+
+				ID3D12Resource* NativeD3D12Resource = GetID3D12DynamicRHI()->RHIGetResource(InBuffer);
+				BufferMemorySize = GetID3D12DynamicRHI()->RHIGetResourceMemorySize(InBuffer);
+
+				TRefCountPtr<ID3D12Device> OwnerDevice;
+				HRESULT QueryResult;
+				if ((QueryResult = NativeD3D12Resource->GetDevice(IID_PPV_ARGS(OwnerDevice.GetInitReference()))) != S_OK)
+				{
+					UE_LOG(LogRivermax, Error, TEXT("Failed to get D3D12 device for captured buffer ressource: %d)"), QueryResult);
+					return nullptr;
+				}
+
+				if ((QueryResult = OwnerDevice->CreateSharedHandle(NativeD3D12Resource, NULL, GENERIC_ALL, NULL, &D3D12BufferHandle)) != S_OK)
+				{
+					UE_LOG(LogRivermax, Error, TEXT("Failed to create shared handle for captured buffer ressource: %d"), QueryResult);
+					return nullptr;
+				}
+
+				CudaExtMemHandleDesc.type = CU_EXTERNAL_MEMORY_HANDLE_TYPE_D3D12_RESOURCE;
+				CudaExtMemHandleDesc.handle.win32.name = nullptr;
+				CudaExtMemHandleDesc.handle.win32.handle = D3D12BufferHandle;
+				CudaExtMemHandleDesc.size = BufferMemorySize;
+				CudaExtMemHandleDesc.flags |= CUDA_EXTERNAL_MEMORY_DEDICATED;
+			}
+
+			FCUDAModule& CudaModule = FModuleManager::GetModuleChecked<FCUDAModule>("CUDA");
+
+			CudaModule.DriverAPI()->cuCtxPushCurrent(CudaModule.GetCudaContext());
+
+			{
+				TRACE_CPUPROFILER_EVENT_SCOPE(Rmax_CudaImportMemory);
+
+				const CUresult Result = FCUDAModule::CUDA().cuImportExternalMemory(&MappedExternalMemory, &CudaExtMemHandleDesc);
+
+				if (D3D12BufferHandle)
+				{
+					CloseHandle(D3D12BufferHandle);
+				}
+
+				if (Result != CUDA_SUCCESS)
+				{
+					UE_LOG(LogRivermax, Error, TEXT("Failed to import shared buffer. Error: %d"), Result);
+					return nullptr;
+				}
+			}
+
+			{
+				TRACE_CPUPROFILER_EVENT_SCOPE(Rmax_MapCudaMemory);
+
+				CUDA_EXTERNAL_MEMORY_BUFFER_DESC BufferDescription = {};
+				BufferDescription.offset = 0;
+				BufferDescription.size = BufferMemorySize;
+				CUdeviceptr NewMemory;
+				const CUresult Result = FCUDAModule::CUDA().cuExternalMemoryGetMappedBuffer(&NewMemory, MappedExternalMemory, &BufferDescription);
+				if (Result != CUDA_SUCCESS || NewMemory == 0)
+				{
+					UE_LOG(LogRivermax, Error, TEXT("Failed to get shared buffer mapped memory. Error: %d"), Result);
+					CudaModule.DriverAPI()->cuCtxPushCurrent(nullptr);
+					return nullptr;
+				}
+
+				BufferGPUMemoryMap.Add(InBuffer, reinterpret_cast<void*>(NewMemory));
+			}
+
+			CudaModule.DriverAPI()->cuCtxPushCurrent(nullptr);
+		}
+
+		// At this point, we have the mapped buffer in cuda space and we can use it to schedule a memcpy on cuda engine.
+		return BufferGPUMemoryMap[InBuffer];
+	}
+
+	void FRivermaxInputStream::ProcessSRD(const FRTPHeader& RTPHeader, uint8* DataPtr)
+	{
+		uint32 DataOffset = 0;
+		StreamData.ReceivedSize += RTPHeader.SRD1.Length;
+		if (bIsUsingGPUDirect)
+		{
+			// Initial case (start address)
+			if (StreamData.DeviceWritePointerOne == nullptr)
+			{
+				StreamData.DeviceWritePointerOne = DataPtr;
+				StreamData.SizeToWriteOne = RTPHeader.SRD1.Length;
+			}
+			else
+			{
+				// Detection of wrap around -> Move tracking to second buffer
+				if (StreamData.DeviceWritePointerTwo == nullptr && DataPtr < StreamData.DeviceWritePointerOne)
+				{
+					StreamData.DeviceWritePointerTwo = DataPtr;
+					StreamData.SizeToWriteTwo = 0;
+				}
+
+				// Case where we track memory in first buffer
+				if (StreamData.DeviceWritePointerTwo == nullptr)
+				{
+					StreamData.SizeToWriteOne += RTPHeader.SRD1.Length;
+				}
+				else // Tracking memory in second buffer
+				{
+					StreamData.SizeToWriteTwo += RTPHeader.SRD1.Length;
+				}
+			}
+		}
+		else
+		{
+			uint8* WriteBuffer = reinterpret_cast<uint8*>(StreamData.CurrentFrame);
+			FMemory::Memcpy(&WriteBuffer[StreamData.WritingOffset], &DataPtr[DataOffset], RTPHeader.SRD1.Length);
+			StreamData.WritingOffset += RTPHeader.SRD1.Length;
+
+			if (RTPHeader.SRD1.bHasContinuation)
+			{
+				DataOffset += RTPHeader.SRD1.Length;
+				FMemory::Memcpy(&WriteBuffer[StreamData.WritingOffset], &DataPtr[DataOffset], RTPHeader.SRD2.Length);
+				StreamData.WritingOffset += RTPHeader.SRD2.Length;
+				StreamData.ReceivedSize += RTPHeader.SRD2.Length;
+			}
+		}
+	}
+
+	void FRivermaxInputStream::ProcessLastSRD(const FRTPHeader& RTPHeader, uint8* DataPtr)
+	{
+		if (StreamData.ReceivedSize == StreamData.ExpectedSize)
+		{
+			const FString TraceName = FString::Format(TEXT("RmaxInput::ReceivedFrame {0}"), { RTPHeader.Timestamp });
+			TRACE_CPUPROFILER_EVENT_SCOPE_TEXT(*TraceName);
+
+			++StreamStats.FramesReceived;
+
+			FRivermaxInputVideoFrameDescriptor Descriptor;
+			Descriptor.Width = Options.Resolution.X;
+			Descriptor.Height = Options.Resolution.Y;
+			Descriptor.Stride = Options.AlignedResolution.X / FormatInfo.PixelGroupCoverage * FormatInfo.PixelGroupSize;
+			Descriptor.bIsUsingGPUDirect = bIsUsingGPUDirect;
+			FRivermaxInputVideoFrameReception NewFrame;
+
+			if (bIsUsingGPUDirect)
+			{
+				if (LastSRDLength.IsSet() && LastSRDLength.GetValue() != RTPHeader.SRD1.Length)
+				{
+					UE_LOG(LogRivermax, Warning, TEXT("Unsupported variable SRD length detected while GPUDirect for input stream is used. Disable and reopen the stream. (Last: %d, New: %d)"), LastSRDLength.GetValue(), RTPHeader.SRD1.Length);
+					Listener->OnStreamError();
+					bIsShuttingDown = true;
+				}
+				LastSRDLength = RTPHeader.SRD1.Length;
+
+				//Frame received entirely, time to copy it from rivermax gpu scratchpad to our own gpu memory
+				FCUDAModule& CudaModule = FModuleManager::GetModuleChecked<FCUDAModule>("CUDA");
+				CUresult Result = CudaModule.DriverAPI()->cuCtxPushCurrent(CudaModule.GetCudaContext());
+
+				const CUdeviceptr DestinationGPUMemory = reinterpret_cast<CUdeviceptr>(StreamData.CurrentFrame);
+				const CUdeviceptr SourceGPUMemoryOne = reinterpret_cast<CUdeviceptr>(StreamData.DeviceWritePointerOne);
+
+				const uint32 NumSRDPartOne = RTPHeader.SRD1.Length > 0 ? StreamData.SizeToWriteOne / RTPHeader.SRD1.Length : 0;
+				const uint32 NumSRDPartTwo = RTPHeader.SRD1.Length > 0 ? StreamData.SizeToWriteTwo / RTPHeader.SRD1.Length : 0;
+
+				// Use cuda's 2d memcopy to do a source and destination stride difference memcopy
+				// We initialize rivermax stream with a payload size blindly since we don't know what the sender will use
+				// So, we use a big value by default and expect SRD to be smaller
+				// This memcopy will consume the SRD size value but jump the init payload size value on the source address
+				// Limitation is that this will only work for fixed SRD across a frame
+				CUDA_MEMCPY2D StrideDescription;
+				FMemory::Memset(StrideDescription, 0);
+				StrideDescription.srcDevice = SourceGPUMemoryOne;
+				StrideDescription.dstDevice = DestinationGPUMemory;
+				StrideDescription.dstMemoryType = CUmemorytype::CU_MEMORYTYPE_DEVICE;
+				StrideDescription.srcMemoryType = CUmemorytype::CU_MEMORYTYPE_DEVICE;
+				StrideDescription.srcPitch = ExpectedPayloadSize; //Source pitch is the expected payload used at init
+				StrideDescription.dstPitch = RTPHeader.SRD1.Length; //Destination pitch is the fixed SRD size we received
+				StrideDescription.WidthInBytes = RTPHeader.SRD1.Length; //Width in bytes is the amount to copy, the SRD size
+				StrideDescription.Height = NumSRDPartOne;
+				Result = CudaModule.DriverAPI()->cuMemcpy2DAsync(&StrideDescription, reinterpret_cast<CUstream>(GPUStream));
+
+				if (StreamData.DeviceWritePointerTwo != nullptr && StreamData.SizeToWriteTwo > 0)
+				{
+					StrideDescription.srcDevice = reinterpret_cast<CUdeviceptr>(StreamData.DeviceWritePointerTwo);
+					StrideDescription.dstDevice = reinterpret_cast<CUdeviceptr>(StreamData.CurrentFrame) + StreamData.SizeToWriteOne;
+					StrideDescription.Height = NumSRDPartTwo;
+					Result = CudaModule.DriverAPI()->cuMemcpy2DAsync(&StrideDescription, reinterpret_cast<CUstream>(GPUStream));
+				}
+
+				if (Result != CUDA_SUCCESS)
+				{
+					UE_LOG(LogRivermax, Warning, TEXT("Failed to copy received buffer to shared memory. Error: %d"), Result);
+					bIsFirstFrameReceived = false;
+					return;
+				}
+
+				auto CudaCallback = [](void* userData)
+				{
+					TRACE_CPUPROFILER_EVENT_SCOPE(FRivermaxInputStream::MemcopyCallback);
+					if (userData)
+					{
+						// It might happen that our stream has been closed once the callback is triggered
+						TWeakPtr<FCallbackPayload>* WeakPayloadPtr = reinterpret_cast<TWeakPtr<FCallbackPayload>*>(userData);
+						TWeakPtr<FCallbackPayload> WeakPayload = *WeakPayloadPtr;
+						if (TSharedPtr<FCallbackPayload> Payload = WeakPayload.Pin())
+						{
+							Payload->bIsWaitingForPendingCopy = false;
+						}
+					}
+				};
+
+				// Schedule a callback to know when to make the frame available
+				CallbackPayload->bIsWaitingForPendingCopy = true;
+				TWeakPtr<FCallbackPayload> WeakPayload = CallbackPayload;
+				CudaModule.DriverAPI()->cuLaunchHostFunc(reinterpret_cast<CUstream>(GPUStream), CudaCallback, &WeakPayload);
+
+				FCUDAModule::CUDA().cuCtxPopCurrent(nullptr);
+
+				// For now, we wait for the cuda callback before we move on receiving next frame. 
+				// Will need to update this and make the frame available from the cuda callback to avoid losing packets 
+				{
+					TRACE_CPUPROFILER_EVENT_SCOPE(FRivermaxInputStream:WaitingPendingOperation)
+					const double CallbackTimestamp = FPlatformTime::Seconds();
+					while (CallbackPayload->bIsWaitingForPendingCopy == true && bIsShuttingDown == false)
+					{
+						if (FPlatformTime::Seconds() - CallbackTimestamp > CVarWaitForCompletionTimeout.GetValueOnAnyThread())
+						{
+							Listener->OnStreamError();
+							CallbackPayload->bIsWaitingForPendingCopy = false;
+							break;
+						}
+					}
+				}
+			}
+
+			NewFrame.VideoBuffer = reinterpret_cast<uint8*>(StreamData.CurrentFrame);
+
+			// No need to provide the new frame and prepare the next one if we are shutting down
+			if (bIsShuttingDown == false)
+			{
+				Listener->OnVideoFrameReceived(Descriptor, NewFrame);
+				PrepareNextFrame();
+			}
+		}
+		else
+		{
+			UE_LOG(LogRivermax, Warning, TEXT("End of frame received (Marker bit) but not enough data was received (missing %d). Expected %d but received (%d)"), StreamData.ExpectedSize - StreamData.ReceivedSize, StreamData.ExpectedSize, StreamData.ReceivedSize);
+			StreamData.WritingOffset = 0;
+			StreamData.ReceivedSize = 0;
+			++StreamStats.InvalidFramesCount;
+		}
+	}
 }
 
diff --git a/Engine/Plugins/VirtualProduction/Rivermax/RivermaxCore/Source/RivermaxCore/Private/RivermaxInputStream.h b/Engine/Plugins/VirtualProduction/Rivermax/RivermaxCore/Source/RivermaxCore/Private/RivermaxInputStream.h
--- a/Engine/Plugins/VirtualProduction/Rivermax/RivermaxCore/Source/RivermaxCore/Private/RivermaxInputStream.h
+++ b/Engine/Plugins/VirtualProduction/Rivermax/RivermaxCore/Source/RivermaxCore/Private/RivermaxInputStream.h
@@ -7,6 +7,7 @@
 #include <atomic>
 #include "Async/Future.h"
 #include "HAL/Runnable.h"
+#include "RHI.h"
 #include "RivermaxHeader.h"
 #include "RivermaxTypes.h"
 
@@ -17,6 +18,9 @@ namespace UE::RivermaxCore::Private
 	using UE::RivermaxCore::IRivermaxInputStream;
 	using UE::RivermaxCore::IRivermaxInputStreamListener;
 	using UE::RivermaxCore::FRivermaxStreamOptions;
+
+	struct FRawRTPHeader;
+	struct FRTPHeader;
 	
 	struct FInputStreamBufferConfiguration
 	{
@@ -46,10 +50,16 @@ namespace UE::RivermaxCore::Private
 	struct FInputStreamData
 	{
 		uint64 LastSequenceNumber = 0;
-		uint8* CurrentFrame = nullptr;
+		void* CurrentFrame = nullptr;
 		uint32 WritingOffset = 0;
 		uint32 ReceivedSize = 0;
 		uint32 ExpectedSize = 0;
+
+		// GPUDirect data handling. Rivermax uses a cyclic buffer so we might have part of the frame at the end of the first buffer and the rest on the second part.
+		void* DeviceWritePointerOne = nullptr;
+		void* DeviceWritePointerTwo = nullptr;
+		uint32 SizeToWriteOne = 0;
+		uint32 SizeToWriteTwo = 0;
 	};
 
 	class FRivermaxInputStream : public IRivermaxInputStream, public FRunnable
@@ -76,19 +86,16 @@ namespace UE::RivermaxCore::Private
 
 	private:
 
-		struct FRTPParameter
-		{
-			uint32 SequencerNumber = 0;
-			uint32 Timestamp = 0;
-			bool bIsMBit = false;
-			bool bIsFBit = false;
-		};
-
-		bool GetRTPParameter(uint8* InputRTP, FRTPParameter& OutParameter);
-		void ParseChunk(const rmax_in_completion& Completion);
-		bool IsExtendedSequenceNumber() const;
+		bool TranslateRTPHeader(const FRawRTPHeader& RawHeader, FRTPHeader& OutParameter);
+		void ParseChunks(const rmax_in_completion& Completion);
+		void ProcessSRD(const FRTPHeader& RTPHeader, uint8* DataPtr);
+		void ProcessLastSRD(const FRTPHeader& RTPHeader, uint8* DataPtr);
 		void PrepareNextFrame();
 		void LogStats();
+		void AllocateBuffers();
+		bool AllocateGPUBuffers();
+		void DeallocateBuffers();
+		void* GetMappedBuffer(const FBufferRHIRef& InBuffer);
 
 	private:
 
@@ -103,6 +110,7 @@ namespace UE::RivermaxCore::Private
 		rmax_in_flow_attr FlowAttribute;
 		FInputStreamBufferConfiguration BufferConfiguration;
 		bool bIsFirstFrameReceived = false;
+		bool bIsFirstPacketReceived = false;
 		ERivermaxStreamType RivermaxStreamType = ERivermaxStreamType::VIDEO_2110_20_STREAM;
 		FInputStreamData StreamData;
 		FInputStreamStats StreamStats;
@@ -111,6 +119,33 @@ namespace UE::RivermaxCore::Private
 		TFuture<void> InitTaskFuture;
 		double LastLoggingTimestamp = 0.0;
 		FVideoFormatInfo FormatInfo;
+		bool bIsUsingGPUDirect = false;
+		int32 DeviceIndex = 0;
+
+		/** Allocated memory base address used when it's time to free */
+		void* GPUAllocatedMemoryBaseAddress = nullptr;
+
+		/** Total allocated gpu memory. */
+		int32 GPUAllocatedMemorySize = 0;
+
+		/** Map between buffer we are writing to and their mapped address in cuda space */
+		TMap<FRHIBuffer*, void*> BufferGPUMemoryMap;
+
+		/** Used to wait for a frame to be copied from gpu to gpu memory */
+		struct FCallbackPayload
+		{
+			volatile bool bIsWaitingForPendingCopy = false;
+		};
+		TSharedPtr<FCallbackPayload> CallbackPayload;
+
+		/** Cuda stream used for our operations */ 
+		void* GPUStream = nullptr;
+
+		/** Payload size used to initialize rivermax stream and know differences between received SRD and config */
+		uint32 ExpectedPayloadSize = 1500;
+
+		/** Used to track SRD length received and detect a change across a frame, which doesn't work with gpudirect */
+		TOptional<uint16> LastSRDLength = 0;
 	};
 }
 
diff --git a/Engine/Plugins/VirtualProduction/Rivermax/RivermaxCore/Source/RivermaxCore/Private/RivermaxOutputStream.cpp b/Engine/Plugins/VirtualProduction/Rivermax/RivermaxCore/Source/RivermaxCore/Private/RivermaxOutputStream.cpp
--- a/Engine/Plugins/VirtualProduction/Rivermax/RivermaxCore/Source/RivermaxCore/Private/RivermaxOutputStream.cpp
+++ b/Engine/Plugins/VirtualProduction/Rivermax/RivermaxCore/Source/RivermaxCore/Private/RivermaxOutputStream.cpp
@@ -183,7 +183,12 @@ namespace UE::RivermaxCore::Private
 						StreamData.FrameFieldTimeIntervalNs = 1E9 / Options.FrameRate.AsDecimal();
 						InitializeStreamTimingSettings();
 
-						UE_LOG(LogRivermax, Display, TEXT("Output started to send %dx%d using %d payloads of size %d"), Options.AlignedResolution.X, Options.AlignedResolution.Y, StreamMemory.PacketsInLine, StreamMemory.PayloadSize);
+						UE_LOG(LogRivermax, Display, TEXT("Output started to send %dx%d using %d payloads of size %d%s")
+						, Options.AlignedResolution.X
+						, Options.AlignedResolution.Y
+						, StreamMemory.PacketsInLine
+						, StreamMemory.PayloadSize
+						, bUseGPUDirect ? TEXT(" using GPUDirect") : TEXT(""));
 
 						bIsActive = true;
 						RivermaxThread.Reset(FRunnableThread::Create(this, TEXT("Rmax OutputStream Thread"), 128 * 1024, TPri_AboveNormal, FPlatformAffinity::GetPoolThreadMask()));
@@ -277,10 +282,17 @@ namespace UE::RivermaxCore::Private
 			{
 				return;
 			}
+			
+			// Add markup when we start pushing out a frame with its timestamp to track it across network
+			if (StreamData.bHasFrameFirstChunkBeenFetched == false)
+			{
+				const uint32 MediaTimestamp = GetTimestampFromTime(StreamData.NextAlignmentPointNanosec, MediaClockSampleRate);
+				const FString TraceName = FString::Format(TEXT("RmaxOutput::StartSending {0}"), { MediaTimestamp });
+				TRACE_CPUPROFILER_EVENT_SCOPE_TEXT(*TraceName);
+			}
 
-			const FString TraceName = FString::Format(TEXT("FRivermaxOutputStream::SendFrame {0}"), { CurrentFrame->FrameIndex });
+			const FString TraceName = FString::Format(TEXT("RmaxOutput::SendFrame {0}"), { CurrentFrame->FrameIndex });
 			TRACE_CPUPROFILER_EVENT_SCOPE_TEXT(*TraceName);
-
 			do 
 			{
 				if (bIsActive)
@@ -474,7 +486,7 @@ namespace UE::RivermaxCore::Private
 		return NextFrame;
 	}
 
-	void FRivermaxOutputStream::BuildRTPHeader(FRTPHeader& OutHeader) const
+	void FRivermaxOutputStream::BuildRTPHeader(FOutputRTPHeader& OutHeader) const
 	{
 		// build RTP header - 12 bytes
 		/*
@@ -722,7 +734,7 @@ namespace UE::RivermaxCore::Private
 		for (uint32 StrideIndex = 0; StrideIndex < StreamMemory.ChunkSizeInStrides && CurrentFrame->PacketCounter < StreamMemory.PacketsPerFrame; ++StrideIndex)
 		{
 			uint8* NextHeaderRawPtr = HeaderRawPtr + (StrideIndex * StreamMemory.HeaderStrideSize);
-			BuildRTPHeader(*reinterpret_cast<FRTPHeader*>(NextHeaderRawPtr));
+			BuildRTPHeader(*reinterpret_cast<FOutputRTPHeader*>(NextHeaderRawPtr));
 			//todo only for video
 			if (!((StrideIndex + 1) % StreamMemory.PacketsInLine))
 			{
@@ -990,7 +1002,7 @@ namespace UE::RivermaxCore::Private
 			}
 
 			const CUdeviceptr CudaMemoryPointer = reinterpret_cast<CUdeviceptr>(MappedPointer);
-			Result = CudaModule.DriverAPI()->cuMemcpyDtoD(reinterpret_cast<CUdeviceptr>(AvailableFrame->VideoBuffer), CudaMemoryPointer, Options.AlignedResolution.Y * Stride);
+			Result = CudaModule.DriverAPI()->cuMemcpyDtoDAsync(reinterpret_cast<CUdeviceptr>(AvailableFrame->VideoBuffer), CudaMemoryPointer, Options.AlignedResolution.Y * Stride, reinterpret_cast<CUstream>(GPUStream));
 			if (Result != CUDA_SUCCESS)
 			{
 				UE_LOG(LogRivermax, Error, TEXT("Failed to copy captured bufer to cuda memory. Stopping capture. Error: %d"), Result);
@@ -1002,7 +1014,7 @@ namespace UE::RivermaxCore::Private
 			// Callback called by Cuda when stream work has completed on cuda engine (MemCpy -> Callback)
 			// Once Memcpy has been done, we know we can mark that memory as available to be sent. 
 			
-			auto CudaCallback = [](CUstream hStream, CUresult status, void* userData)
+			auto CudaCallback = [](void* userData)
 			{
 				TRACE_CPUPROFILER_EVENT_SCOPE(CudaWorkDoneCallback);
 
@@ -1033,7 +1045,7 @@ namespace UE::RivermaxCore::Private
 			PendingIdentifiers.Enqueue(AvailableFrame->FrameIdentifier);
 
 			// Schedule a callback to make the frame available
-			CudaModule.DriverAPI()->cuStreamAddCallback(0, CudaCallback, this, 0);
+			CudaModule.DriverAPI()->cuLaunchHostFunc(reinterpret_cast<CUstream>(GPUStream), CudaCallback, this);
 			
 			FCUDAModule::CUDA().cuCtxPopCurrent(nullptr);
 
@@ -1157,6 +1169,16 @@ namespace UE::RivermaxCore::Private
 			UE_LOG(LogRivermax, Warning, TEXT("Can't initialize output to use GPUDirect. Failed to configure memory access. Status: %d"), Status);
 			return false;
 		}
+
+		CUstream CudaStream;
+		Status = CudaModule.DriverAPI()->cuStreamCreate(&CudaStream, CU_STREAM_NON_BLOCKING);
+		if (Status != CUDA_SUCCESS)
+		{
+			UE_LOG(LogRivermax, Warning, TEXT("Can't initialize output to use GPUDirect. Failed to create its stream. Status: %d"), Status);
+			return false;
+		}
+
+		GPUStream = CudaStream;
 		
 		Status = CudaModule.DriverAPI()->cuCtxSynchronize();
 		if (Status != CUDA_SUCCESS)
@@ -1209,6 +1231,13 @@ namespace UE::RivermaxCore::Private
 			}
 			BufferCudaMemoryMap.Empty();
 
+			Status = CudaModule.DriverAPI()->cuStreamDestroy(reinterpret_cast<CUstream>(GPUStream));
+			if (Status != CUDA_SUCCESS)
+			{
+				UE_LOG(LogRivermax, Warning, TEXT("Failed to destroy cuda stream. Status: %d"), Status);
+			}
+			GPUStream = nullptr;
+
 			CudaModule.DriverAPI()->cuCtxPopCurrent(nullptr);
 		}
 	}
diff --git a/Engine/Plugins/VirtualProduction/Rivermax/RivermaxCore/Source/RivermaxCore/Private/RivermaxOutputStream.h b/Engine/Plugins/VirtualProduction/Rivermax/RivermaxCore/Source/RivermaxCore/Private/RivermaxOutputStream.h
--- a/Engine/Plugins/VirtualProduction/Rivermax/RivermaxCore/Source/RivermaxCore/Private/RivermaxOutputStream.h
+++ b/Engine/Plugins/VirtualProduction/Rivermax/RivermaxCore/Source/RivermaxCore/Private/RivermaxOutputStream.h
@@ -19,7 +19,7 @@ namespace UE::RivermaxCore::Private
 	using UE::RivermaxCore::FRivermaxOutputVideoFrameInfo;
 	
 	// Todo make a proper RTP header struct
-	struct FRTPHeader
+	struct FOutputRTPHeader
 	{
 		uint8 RawHeader[20];
 	};
@@ -48,7 +48,7 @@ namespace UE::RivermaxCore::Private
 		TArray<rmax_mem_block> MemoryBlocks;
 		TArray<uint16_t> PayloadSizes; //Array describing stride payload size
 		TArray<uint16_t> HeaderSizes; //Array describing header payload size
-		TArray<FRTPHeader> RTPHeaders;
+		TArray<FOutputRTPHeader> RTPHeaders;
 
 		rmax_buffer_attr BufferAttributes;
 	};
@@ -117,7 +117,7 @@ namespace UE::RivermaxCore::Private
 		void InitializeNextFrame(const TSharedPtr<FRivermaxOutputFrame>& NextFrame);
 		TSharedPtr<FRivermaxOutputFrame> GetNextFrameToSend();
 		TSharedPtr<FRivermaxOutputFrame> GetNextAvailableFrame(uint32 InFrameIdentifier);
-		void BuildRTPHeader(FRTPHeader& OutHeader) const;
+		void BuildRTPHeader(FOutputRTPHeader& OutHeader) const;
 		void DestroyStream();
 		void WaitForNextRound();
 		void GetNextChunk();
@@ -204,6 +204,9 @@ namespace UE::RivermaxCore::Private
 
 		/** Queued identifiers to be consumed by cuda callback when work has been completed */
 		TSpscQueue<uint32> PendingIdentifiers;
+
+		/** Cuda stream used for our operations */
+		void* GPUStream = nullptr;
 	};
 }
 
diff --git a/Engine/Plugins/VirtualProduction/Rivermax/RivermaxCore/Source/RivermaxCore/Public/IRivermaxInputStream.h b/Engine/Plugins/VirtualProduction/Rivermax/RivermaxCore/Source/RivermaxCore/Public/IRivermaxInputStream.h
--- a/Engine/Plugins/VirtualProduction/Rivermax/RivermaxCore/Source/RivermaxCore/Public/IRivermaxInputStream.h
+++ b/Engine/Plugins/VirtualProduction/Rivermax/RivermaxCore/Source/RivermaxCore/Public/IRivermaxInputStream.h
@@ -4,6 +4,8 @@
 
 #include "CoreMinimal.h"
 
+class FRHIBuffer;
+
 namespace UE::RivermaxCore
 {
 	struct FRivermaxStreamOptions;
@@ -14,11 +16,16 @@ namespace UE::RivermaxCore
 		uint32 Width = 0;
 		uint32 Stride = 0;
 		uint32 VideoBufferSize = 0;
+		bool bIsUsingGPUDirect = false;
 	};
 
 	struct RIVERMAXCORE_API FRivermaxInputVideoFrameRequest
 	{
+		/** Buffer pointer in RAM where to write incoming frame */
 		uint8* VideoBuffer = nullptr;
+		
+		/** Buffer pointer in GPU to be mapped to cuda and where to write incoming frame */
+		FRHIBuffer* GPUBuffer = nullptr;
 	};
 
 	struct RIVERMAXCORE_API FRivermaxInputVideoFrameReception
@@ -37,6 +44,9 @@ namespace UE::RivermaxCore
 		
 		/** Called when a frame has been received */
 		virtual void OnVideoFrameReceived(const FRivermaxInputVideoFrameDescriptor& FrameInfo, const FRivermaxInputVideoFrameReception& ReceivedVideoFrame) = 0;
+
+		/** Called when stream has encountered an error and has to stop */
+		virtual void OnStreamError() = 0;
 	};
 
 	class RIVERMAXCORE_API IRivermaxInputStream
diff --git a/Engine/Plugins/VirtualProduction/Rivermax/RivermaxMedia/Source/RivermaxMedia/Private/RivermaxMediaPlayer.cpp b/Engine/Plugins/VirtualProduction/Rivermax/RivermaxMedia/Source/RivermaxMedia/Private/RivermaxMediaPlayer.cpp
--- a/Engine/Plugins/VirtualProduction/Rivermax/RivermaxMedia/Source/RivermaxMedia/Private/RivermaxMediaPlayer.cpp
+++ b/Engine/Plugins/VirtualProduction/Rivermax/RivermaxMedia/Source/RivermaxMedia/Private/RivermaxMediaPlayer.cpp
@@ -9,6 +9,7 @@
 #include "MediaIOCoreFileWriter.h"
 #include "MediaIOCoreSamples.h"
 #include "Misc/ScopeLock.h"
+#include "RenderCommandFence.h"
 #include "RivermaxMediaLog.h"
 #include "RivermaxMediaOutput.h"
 #include "RivermaxMediaSourceOptions.h"
@@ -88,7 +89,7 @@ namespace UE::RivermaxMedia
 		}
 		if (InputStream == nullptr || !InputStream->Initialize(StreamOptions, *this))
 		{
-			UE_LOG(LogRivermaxMedia, Warning, TEXT("The Rivermax port couldn't be opened."));
+			UE_LOG(LogRivermaxMedia, Warning, TEXT("Failed to initialize Rivermax input stream."));
 			CurrentState = EMediaState::Error;
 			RivermaxThreadNewState = EMediaState::Error;
 			InputStream.Reset();
@@ -98,6 +99,8 @@ namespace UE::RivermaxMedia
 		// Setup our different supported channels based on source settings
 		SetupSampleChannels();
 
+		AllocateBuffers();
+
 		// finalize
 		CurrentState = EMediaState::Preparing;
 		RivermaxThreadNewState = EMediaState::Preparing;
@@ -165,11 +168,36 @@ namespace UE::RivermaxMedia
 
 	bool FRivermaxMediaPlayer::OnVideoFrameRequested(const FRivermaxInputVideoFrameDescriptor& FrameInfo, FRivermaxInputVideoFrameRequest& OutVideoFrameRequest)
 	{
+		TRACE_CPUPROFILER_EVENT_SCOPE(FRivermaxMediaPlayer::OnVideoFrameRequested)
+
+		if (RivermaxThreadNewState != EMediaState::Playing)
+		{
+			return false;
+		}
+
 		if (FrameInfo.VideoBufferSize > 0)
 		{
-			RivermaxThreadCurrentTextureSample = TextureSamplePool->AcquireShared();
-			OutVideoFrameRequest.VideoBuffer = reinterpret_cast<uint8*>(RivermaxThreadCurrentTextureSample->RequestBuffer(FrameInfo.VideoBufferSize));
-			return OutVideoFrameRequest.VideoBuffer != nullptr;
+			if (FrameInfo.bIsUsingGPUDirect)
+			{
+				RivermaxThreadCurrentTextureSample = TextureSamplePool->AcquireShared();
+				if (RivermaxThreadCurrentTextureSample->GetGPUBuffer() == nullptr)
+				{
+					UE_LOG(LogRivermaxMedia, Verbose, TEXT("Video stream overrunning engine. Allocating a new frame."));
+					RivermaxThreadCurrentTextureSample->InitializeGPUBuffer(StreamOptions.AlignedResolution, DesiredPixelFormat);
+				}
+				else
+				{
+					OutVideoFrameRequest.GPUBuffer = RivermaxThreadCurrentTextureSample->GetGPUBuffer()->GetRHI();
+				}
+				
+				return OutVideoFrameRequest.GPUBuffer != nullptr;
+			}
+			else
+			{
+				RivermaxThreadCurrentTextureSample = TextureSamplePool->AcquireShared();
+				OutVideoFrameRequest.VideoBuffer = reinterpret_cast<uint8*>(RivermaxThreadCurrentTextureSample->RequestBuffer(FrameInfo.VideoBufferSize));
+				return OutVideoFrameRequest.VideoBuffer != nullptr;
+			}
 		}
 
 		return false;
@@ -178,6 +206,13 @@ namespace UE::RivermaxMedia
 
 	void FRivermaxMediaPlayer::OnVideoFrameReceived(const FRivermaxInputVideoFrameDescriptor& FrameInfo, const FRivermaxInputVideoFrameReception& ReceivedVideoFrame)
 	{
+		TRACE_CPUPROFILER_EVENT_SCOPE(FRivermaxMediaPlayer::OnVideoFrameReceived)
+
+		if (RivermaxThreadNewState != EMediaState::Playing)
+		{
+			return;
+		}
+
 		FTimespan DecodedTime = FTimespan::FromSeconds(GetPlatformSeconds());
 		TOptional<FTimecode> DecodedTimecode;
 
@@ -299,6 +334,7 @@ namespace UE::RivermaxMedia
 
 		StreamOptions.StreamAddress = Options->GetMediaOption(RivermaxMediaOption::StreamAddress, FString());
 		StreamOptions.Port = Options->GetMediaOption(RivermaxMediaOption::Port, (int64)0);
+		StreamOptions.bUseGPUDirect = Options->GetMediaOption(RivermaxMediaOption::UseGPUDirect, false);
 		StreamOptions.Resolution = VideoTrackFormat.Dim;
 		StreamOptions.FrameRate = VideoFrameRate;
 
@@ -313,6 +349,36 @@ namespace UE::RivermaxMedia
 		return true;
 	}
 
+	void FRivermaxMediaPlayer::AllocateBuffers()
+	{
+		using namespace UE::RivermaxCore;
+
+		if (StreamOptions.bUseGPUDirect == false)
+		{
+			return;
+		}
+
+		TArray<TSharedPtr<FRivermaxMediaTextureSample>> TempSamples;
+		for (int32 Index = 0; Index < MaxNumVideoFrameBuffer; ++Index)
+		{
+			// Acquire a certain amount of buffer and initialize their gpu buffer to be ready when requested by rivermax stream
+			TSharedPtr<FRivermaxMediaTextureSample> NewSample = TextureSamplePool->AcquireShared();
+			NewSample->InitializeGPUBuffer(StreamOptions.AlignedResolution, DesiredPixelFormat);
+			TempSamples.Add(MoveTemp(NewSample));
+		}
+
+		// Allocation is done on render thread so let's make sure it's completed before pursuing
+		FRenderCommandFence RenderFence;
+		RenderFence.BeginFence();
+		RenderFence.Wait();		
+	}
+
+	void FRivermaxMediaPlayer::OnStreamError()
+	{
+		// If the stream ends up in error, stop the player
+		RivermaxThreadNewState = EMediaState::Error;
+	}
+
 }
 
 #undef LOCTEXT_NAMESPACE
diff --git a/Engine/Plugins/VirtualProduction/Rivermax/RivermaxMedia/Source/RivermaxMedia/Private/RivermaxMediaPlayer.h b/Engine/Plugins/VirtualProduction/Rivermax/RivermaxMedia/Source/RivermaxMedia/Private/RivermaxMediaPlayer.h
--- a/Engine/Plugins/VirtualProduction/Rivermax/RivermaxMedia/Source/RivermaxMedia/Private/RivermaxMediaPlayer.h
+++ b/Engine/Plugins/VirtualProduction/Rivermax/RivermaxMedia/Source/RivermaxMedia/Private/RivermaxMediaPlayer.h
@@ -65,6 +65,7 @@ namespace  UE::RivermaxMedia
 		virtual void OnInitializationCompleted(bool bHasSucceed) override;
 		virtual bool OnVideoFrameRequested(const FRivermaxInputVideoFrameDescriptor& FrameInfo, FRivermaxInputVideoFrameRequest& OutVideoFrameRequest) override;
 		virtual void OnVideoFrameReceived(const FRivermaxInputVideoFrameDescriptor& FrameInfo, const FRivermaxInputVideoFrameReception& ReceivedVideoFrame) override;
+		virtual void OnStreamError() override;
 		//~ End IRivermaxInputStreamListener interface
 
 	protected:
@@ -83,6 +84,7 @@ namespace  UE::RivermaxMedia
 
 	private:
 		bool ConfigureStream(const IMediaOptions* Options);
+		void AllocateBuffers();
 
 	private:
 
diff --git a/Engine/Plugins/VirtualProduction/Rivermax/RivermaxMedia/Source/RivermaxMedia/Private/RivermaxMediaSource.cpp b/Engine/Plugins/VirtualProduction/Rivermax/RivermaxMedia/Source/RivermaxMedia/Private/RivermaxMediaSource.cpp
--- a/Engine/Plugins/VirtualProduction/Rivermax/RivermaxMedia/Source/RivermaxMedia/Private/RivermaxMediaSource.cpp
+++ b/Engine/Plugins/VirtualProduction/Rivermax/RivermaxMedia/Source/RivermaxMedia/Private/RivermaxMediaSource.cpp
@@ -16,6 +16,10 @@ bool URivermaxMediaSource::GetMediaOption(const FName& Key, bool DefaultValue) c
 	{
 		return bIsSRGBInput;
 	}
+	if (Key == RivermaxMediaOption::UseGPUDirect)
+	{
+		return bUseGPUDirect;
+	}
 	return Super::GetMediaOption(Key, DefaultValue);
 }
 
@@ -72,6 +76,7 @@ bool URivermaxMediaSource::HasMediaOption(const FName& Key) const
 		(Key == RivermaxMediaOption::Port) ||
 		(Key == RivermaxMediaOption::PixelFormat) ||
 		(Key == RivermaxMediaOption::SRGBInput) ||
+		(Key == RivermaxMediaOption::UseGPUDirect) ||
 		(Key == FMediaIOCoreMediaOption::FrameRateNumerator) ||
 		(Key == FMediaIOCoreMediaOption::FrameRateDenominator) ||
 		(Key == FMediaIOCoreMediaOption::ResolutionWidth) ||
diff --git a/Engine/Plugins/VirtualProduction/Rivermax/RivermaxMedia/Source/RivermaxMedia/Private/RivermaxMediaSourceOptions.h b/Engine/Plugins/VirtualProduction/Rivermax/RivermaxMedia/Source/RivermaxMedia/Private/RivermaxMediaSourceOptions.h
--- a/Engine/Plugins/VirtualProduction/Rivermax/RivermaxMedia/Source/RivermaxMedia/Private/RivermaxMediaSourceOptions.h
+++ b/Engine/Plugins/VirtualProduction/Rivermax/RivermaxMedia/Source/RivermaxMedia/Private/RivermaxMediaSourceOptions.h
@@ -12,4 +12,5 @@ namespace RivermaxMediaOption
 	static const FName Port("Port");
 	static const FName PixelFormat("PixelFormat");
 	static const FName SRGBInput("SRGBInput");
+	static const FName UseGPUDirect("UseGPUDirect");
 }
diff --git a/Engine/Plugins/VirtualProduction/Rivermax/RivermaxMedia/Source/RivermaxMedia/Private/RivermaxMediaTextureSample.cpp b/Engine/Plugins/VirtualProduction/Rivermax/RivermaxMedia/Source/RivermaxMedia/Private/RivermaxMediaTextureSample.cpp
--- a/Engine/Plugins/VirtualProduction/Rivermax/RivermaxMedia/Source/RivermaxMedia/Private/RivermaxMediaTextureSample.cpp
+++ b/Engine/Plugins/VirtualProduction/Rivermax/RivermaxMedia/Source/RivermaxMedia/Private/RivermaxMediaTextureSample.cpp
@@ -2,7 +2,9 @@
 
 #include "RivermaxMediaTextureSample.h"
 
+#include "RenderGraphUtils.h"
 #include "RivermaxMediaTextureSampleConverter.h"
+#include "RivermaxMediaUtils.h"
 
 FRivermaxMediaTextureSample::FRivermaxMediaTextureSample()
 	: TextureConverter(MakeUnique<FRivermaxMediaTextureSampleConverter>())
@@ -61,3 +63,30 @@ bool FRivermaxMediaTextureSample::ConfigureSample(uint32 InWidth, uint32 InHeigh
 	return Super::SetProperties(InStride, InWidth, InHeight, VideoSampleFormat, InTime, InFrameRate, InTimecode, bInIsSRGBInput);
 }
 
+TRefCountPtr<FRDGPooledBuffer> FRivermaxMediaTextureSample::GetGPUBuffer() const
+{
+	return GPUBuffer;
+}
+
+void FRivermaxMediaTextureSample::InitializeGPUBuffer(const FIntPoint& InResolution, ERivermaxMediaSourcePixelFormat InSampleFormat)
+{
+	using namespace UE::RivermaxMediaUtils::Private;
+
+	const FSourceBufferDesc BufferDescription = GetBufferDescription(InResolution, InSampleFormat);
+
+	FRDGBufferDesc RDGBufferDesc = FRDGBufferDesc::CreateStructuredDesc(BufferDescription.BytesPerElement, BufferDescription.NumberOfElements);
+	
+	// Required to share resource across different graphics API (DX, Cuda)
+	RDGBufferDesc.Usage |= EBufferUsageFlags::Shared;
+
+	TWeakPtr<FRivermaxMediaTextureSample> WeakSample = AsShared();
+	ENQUEUE_RENDER_COMMAND(RivermaxPlayerBufferCreation)(
+	[WeakSample, RDGBufferDesc](FRHICommandListImmediate& CommandList)
+	{
+		if (TSharedPtr<FRivermaxMediaTextureSample> Sample = WeakSample.Pin())
+		{
+			Sample->GPUBuffer = AllocatePooledBuffer(RDGBufferDesc, TEXT("RmaxInput Buffer"));
+		}
+	});
+}
+
diff --git a/Engine/Plugins/VirtualProduction/Rivermax/RivermaxMedia/Source/RivermaxMedia/Private/RivermaxMediaTextureSample.h b/Engine/Plugins/VirtualProduction/Rivermax/RivermaxMedia/Source/RivermaxMedia/Private/RivermaxMediaTextureSample.h
--- a/Engine/Plugins/VirtualProduction/Rivermax/RivermaxMedia/Source/RivermaxMedia/Private/RivermaxMediaTextureSample.h
+++ b/Engine/Plugins/VirtualProduction/Rivermax/RivermaxMedia/Source/RivermaxMedia/Private/RivermaxMediaTextureSample.h
@@ -30,10 +30,20 @@ public:
 
 	bool ConfigureSample(uint32 InWidth, uint32 InHeight, uint32 InStride, ERivermaxMediaSourcePixelFormat InSampleFormat, FTimespan InTime, const FFrameRate& InFrameRate, const TOptional<FTimecode>& InTimecode, bool bInIsSRGBInput);
 
+	/** Initialized a RDG buffer based on the description required. Only useful for gpudirect functionality */
+	void InitializeGPUBuffer(const FIntPoint& InResolution, ERivermaxMediaSourcePixelFormat InSampleFormat);
+	
+	/** Returns RDG allocated buffer */
+	TRefCountPtr<FRDGPooledBuffer> GetGPUBuffer() const;
+
 private:
+	
 
 	/** Texture converted used to handle incoming 2110 formats and convert them to RGB textures the engine handles */
 	TUniquePtr<FRivermaxMediaTextureSampleConverter> TextureConverter;
+	
+	/** Pooled buffer used for gpudirect functionality. Received content will already be on GPU when received from NIC */
+	TRefCountPtr<FRDGPooledBuffer> GPUBuffer;
 };
 
 /*
diff --git a/Engine/Plugins/VirtualProduction/Rivermax/RivermaxMedia/Source/RivermaxMedia/Private/RivermaxMediaTextureSampleConverter.cpp b/Engine/Plugins/VirtualProduction/Rivermax/RivermaxMedia/Source/RivermaxMedia/Private/RivermaxMediaTextureSampleConverter.cpp
--- a/Engine/Plugins/VirtualProduction/Rivermax/RivermaxMedia/Source/RivermaxMedia/Private/RivermaxMediaTextureSampleConverter.cpp
+++ b/Engine/Plugins/VirtualProduction/Rivermax/RivermaxMedia/Source/RivermaxMedia/Private/RivermaxMediaTextureSampleConverter.cpp
@@ -3,9 +3,10 @@
 #include "RivermaxMediaTextureSampleConverter.h"
 
 #include "RenderGraphBuilder.h"
-#include "RivermaxShaders.h"
+#include "RivermaxMediaLog.h"
 #include "RivermaxMediaTextureSample.h"
-
+#include "RivermaxMediaUtils.h"
+#include "RivermaxShaders.h"
 
 DECLARE_GPU_STAT(RivermaxSource_SampleConversion);
 
@@ -22,6 +23,7 @@ bool FRivermaxMediaTextureSampleConverter::Convert(FTexture2DRHIRef& InDestinati
 	TRACE_CPUPROFILER_EVENT_SCOPE(RivermaxSampleConverter::Convert);
 
 	using namespace UE::RivermaxShaders;
+	using namespace UE::RivermaxMediaUtils::Private;
 
 	TSharedPtr<FRivermaxMediaTextureSample> SamplePtr = Sample.Pin();
 	if (SamplePtr.IsValid() == false)
@@ -29,7 +31,7 @@ bool FRivermaxMediaTextureSampleConverter::Convert(FTexture2DRHIRef& InDestinati
 		return false;
 	}
 
-	FIntVector GroupCount;
+	const FSourceBufferDesc SourceBufferDesc = GetBufferDescription(SamplePtr->GetDim(), InputPixelFormat);
 	
 	FRDGBuilder GraphBuilder(FRHICommandListExecutor::GetImmediateCommandList());
 	{
@@ -38,29 +40,36 @@ bool FRivermaxMediaTextureSampleConverter::Convert(FTexture2DRHIRef& InDestinati
 
 		FRDGTextureRef OutputResource = GraphBuilder.RegisterExternalTexture(CreateRenderTarget(InDestinationTexture, TEXT("RivermaxMediaTextureOutputResource")));
 
+		// If we have a valid GPUBuffer, i.e GPUDirect is involved, use that one. Otherwise, take the system buffer and upload it in a new structured buffer.
+		FRDGBufferRef InputBuffer;
+		if (SamplePtr->GetGPUBuffer().IsValid())
+		{
+			InputBuffer = GraphBuilder.RegisterExternalBuffer(SamplePtr->GetGPUBuffer(), TEXT("RMaxGPUBuffer"));
+		}
+		else
+		{
+			TRACE_CPUPROFILER_EVENT_SCOPE(RivermaxSampleConverter::CreateStructuredBuffer);
+			InputBuffer = CreateStructuredBuffer(GraphBuilder, TEXT("RivermaxInputBuffer"), SourceBufferDesc.BytesPerElement, SourceBufferDesc.NumberOfElements, SamplePtr->GetBuffer(), SourceBufferDesc.BytesPerElement * SourceBufferDesc.NumberOfElements, ERDGInitialDataFlags::NoCopy);
+		}
+		
+		const FIntPoint ProcessedOutputDimension = { (int32)SourceBufferDesc.ElementsPerRow, InDestinationTexture->GetDesc().Extent.Y };
+		const FIntVector GroupCount = FComputeShaderUtils::GetGroupCount(ProcessedOutputDimension, FComputeShaderUtils::kGolden2DGroupSize);
+		FGlobalShaderMap* GlobalShaderMap = GetGlobalShaderMap(GMaxRHIFeatureLevel);
+
 		//Configure shader and add conversion pass based on desired pixel format
 		switch (InputPixelFormat)
 		{
 		case ERivermaxMediaSourcePixelFormat::YUV422_8bit:
 		{
 			TRACE_CPUPROFILER_EVENT_SCOPE(RivermaxSampleConverter::YUV8ShaderSetup);
-			const int32 BytesPerElement = sizeof(FYUV8Bit422ToRGBACS::FYUV8Bit422Buffer);
-			const uint32 Stride = SamplePtr->GetStride();
-			const int32 ElementsPerRow = (Stride / BytesPerElement) + ((Stride % BytesPerElement > 0) ? 1 : 0);
-			const int32 ElementCount = ElementsPerRow * InDestinationTexture->GetDesc().Extent.Y;
 
 			FYUV8Bit422ToRGBACS::FPermutationDomain PermutationVector;
 			PermutationVector.Set<FYUV8Bit422ToRGBACS::FSRGBToLinear>(bDoSRGBToLinear);
 
-			FRDGBufferRef InputYUVBuffer = CreateStructuredBuffer(GraphBuilder, TEXT("RivermaxInputBuffer"), BytesPerElement, ElementCount, SamplePtr->GetBuffer(), BytesPerElement * ElementCount);
-			constexpr int32 PixelsPerInput = 2;
-			const FIntPoint ProcessedOutputDimension = { InDestinationTexture->GetDesc().Extent.X / PixelsPerInput, InDestinationTexture->GetDesc().Extent.Y };
-			GroupCount = FComputeShaderUtils::GetGroupCount(ProcessedOutputDimension, FComputeShaderUtils::kGolden2DGroupSize);
-			FGlobalShaderMap* GlobalShaderMap = GetGlobalShaderMap(GMaxRHIFeatureLevel);
+			const FMatrix YUVToRGBMatrix = SamplePtr->GetYUVToRGBMatrix();
+			const FVector YUVOffset(MediaShaders::YUVOffset8bits);
 			TShaderMapRef<FYUV8Bit422ToRGBACS> ComputeShader(GlobalShaderMap, PermutationVector);
-			FMatrix YUVToRGBMatrix = SamplePtr->GetYUVToRGBMatrix();
-			FVector YUVOffset(MediaShaders::YUVOffset8bits);
-			FYUV8Bit422ToRGBACS::FParameters* Parameters = ComputeShader->AllocateAndSetParameters(GraphBuilder, InputYUVBuffer, OutputResource, YUVToRGBMatrix, YUVOffset, ElementsPerRow, ProcessedOutputDimension.Y);
+			FYUV8Bit422ToRGBACS::FParameters* Parameters = ComputeShader->AllocateAndSetParameters(GraphBuilder, InputBuffer, OutputResource, YUVToRGBMatrix, YUVOffset, SourceBufferDesc.ElementsPerRow, ProcessedOutputDimension.Y);
 
 
 			FComputeShaderUtils::AddPass(
@@ -74,24 +83,14 @@ bool FRivermaxMediaTextureSampleConverter::Convert(FTexture2DRHIRef& InDestinati
 		case ERivermaxMediaSourcePixelFormat::YUV422_10bit:
 		{
 			TRACE_CPUPROFILER_EVENT_SCOPE(RivermaxSampleConverter::YUV10ShaderSetup);
-			const int32 BytesPerElement = sizeof(FYUV10Bit422ToRGBACS::FYUV10Bit422LEBuffer);
-			const uint32 Stride = SamplePtr->GetStride();
-			const int32 ElementsPerRow = (Stride / BytesPerElement) + ((Stride % BytesPerElement > 0) ? 1 : 0);
-			const int32 ElementCount = ElementsPerRow * InDestinationTexture->GetDesc().Extent.Y;
-
+			
 			FYUV10Bit422ToRGBACS::FPermutationDomain PermutationVector;
 			PermutationVector.Set<FYUV10Bit422ToRGBACS::FSRGBToLinear>(bDoSRGBToLinear);
 
-			FRDGBufferRef InputYUVBuffer = CreateStructuredBuffer(GraphBuilder, TEXT("RivermaxInputBuffer"), BytesPerElement, ElementCount, SamplePtr->GetBuffer(), BytesPerElement * ElementCount);
-			constexpr int32 PixelsPerInput = 8;
-			const FIntPoint ProcessedOutputDimension = { InDestinationTexture->GetDesc().Extent.X / PixelsPerInput, InDestinationTexture->GetDesc().Extent.Y };
-			GroupCount = FComputeShaderUtils::GetGroupCount(ProcessedOutputDimension, FComputeShaderUtils::kGolden2DGroupSize);
-			FGlobalShaderMap* GlobalShaderMap = GetGlobalShaderMap(GMaxRHIFeatureLevel);
+			const FMatrix YUVToRGBMatrix = SamplePtr->GetYUVToRGBMatrix();
+			const FVector YUVOffset(MediaShaders::YUVOffset10bits);
 			TShaderMapRef<FYUV10Bit422ToRGBACS> ComputeShader(GlobalShaderMap, PermutationVector);
-			FMatrix YUVToRGBMatrix = SamplePtr->GetYUVToRGBMatrix();
-			FVector YUVOffset(MediaShaders::YUVOffset10bits);
-			FYUV10Bit422ToRGBACS::FParameters* Parameters = ComputeShader->AllocateAndSetParameters(GraphBuilder, InputYUVBuffer, OutputResource, YUVToRGBMatrix, YUVOffset, ElementsPerRow, ProcessedOutputDimension.Y);
-
+			FYUV10Bit422ToRGBACS::FParameters* Parameters = ComputeShader->AllocateAndSetParameters(GraphBuilder, InputBuffer, OutputResource, YUVToRGBMatrix, YUVOffset, SourceBufferDesc.ElementsPerRow, ProcessedOutputDimension.Y);
 
 			FComputeShaderUtils::AddPass(
 				GraphBuilder
@@ -104,21 +103,12 @@ bool FRivermaxMediaTextureSampleConverter::Convert(FTexture2DRHIRef& InDestinati
 		case ERivermaxMediaSourcePixelFormat::RGB_8bit:
 		{
 			TRACE_CPUPROFILER_EVENT_SCOPE(RivermaxSampleConverter::RGB8ShaderSetup);
-			const int32 BytesPerElement = sizeof(FRGB8BitToRGBA8CS::FRGB8BitBuffer);
-			const uint32 Stride = SamplePtr->GetStride();
-			const int32 ElementsPerRow = (Stride / BytesPerElement) + ((Stride % BytesPerElement > 0) ? 1 : 0);
-			const int32 ElementCount = ElementsPerRow * InDestinationTexture->GetDesc().Extent.Y;
 
 			FRGB8BitToRGBA8CS::FPermutationDomain PermutationVector;
 			PermutationVector.Set<FRGB8BitToRGBA8CS::FSRGBToLinear>(bDoSRGBToLinear);
 
-			FRDGBufferRef InputRGGBuffer = CreateStructuredBuffer(GraphBuilder, TEXT("RivermaxInputBuffer"), BytesPerElement, ElementCount, SamplePtr->GetBuffer(), BytesPerElement * ElementCount);
-			constexpr int32 PixelsPerInput = 4;
-			const FIntPoint ProcessedOutputDimension = { InDestinationTexture->GetDesc().Extent.X / PixelsPerInput,InDestinationTexture->GetDesc().Extent.Y };
-			GroupCount = FComputeShaderUtils::GetGroupCount(ProcessedOutputDimension, FComputeShaderUtils::kGolden2DGroupSize);
-			FGlobalShaderMap* GlobalShaderMap = GetGlobalShaderMap(GMaxRHIFeatureLevel);
 			TShaderMapRef<FRGB8BitToRGBA8CS> ComputeShader(GlobalShaderMap, PermutationVector);
-			FRGB8BitToRGBA8CS::FParameters* Parameters = ComputeShader->AllocateAndSetParameters(GraphBuilder, InputRGGBuffer, OutputResource, ElementsPerRow, ProcessedOutputDimension.Y);
+			FRGB8BitToRGBA8CS::FParameters* Parameters = ComputeShader->AllocateAndSetParameters(GraphBuilder, InputBuffer, OutputResource, SourceBufferDesc.ElementsPerRow, ProcessedOutputDimension.Y);
 
 
 			FComputeShaderUtils::AddPass(
@@ -133,22 +123,11 @@ bool FRivermaxMediaTextureSampleConverter::Convert(FTexture2DRHIRef& InDestinati
 		{
 			TRACE_CPUPROFILER_EVENT_SCOPE(RivermaxSampleConverter::RGB10ShaderSetup);
 
-			const int32 BytesPerElement = sizeof(FRGB10BitToRGBA10CS::FRGB10BitBuffer);
-			const uint32 Stride = SamplePtr->GetStride();
-			const int32 ElementsPerRow = (Stride / BytesPerElement) + ((Stride % BytesPerElement > 0) ? 1 : 0);
-			const int32 ElementCount = ElementsPerRow * InDestinationTexture->GetDesc().Extent.Y;
-
 			FRGB10BitToRGBA10CS::FPermutationDomain PermutationVector;
 			PermutationVector.Set<FRGB10BitToRGBA10CS::FSRGBToLinear>(bDoSRGBToLinear);
 
-			FRDGBufferRef InputRGGBuffer = CreateStructuredBuffer(GraphBuilder, TEXT("RivermaxInputBuffer"), BytesPerElement, ElementCount, SamplePtr->GetBuffer(), BytesPerElement * ElementCount);
-			constexpr int32 PixelsPerInput = 16;
-			const FIntPoint ProcessedOutputDimension = { InDestinationTexture->GetDesc().Extent.X / PixelsPerInput,InDestinationTexture->GetDesc().Extent.Y };
-			GroupCount = FComputeShaderUtils::GetGroupCount(ProcessedOutputDimension, FComputeShaderUtils::kGolden2DGroupSize);
-			FGlobalShaderMap* GlobalShaderMap = GetGlobalShaderMap(GMaxRHIFeatureLevel);
 			TShaderMapRef<FRGB10BitToRGBA10CS> ComputeShader(GlobalShaderMap, PermutationVector);
-			FRGB10BitToRGBA10CS::FParameters* Parameters = ComputeShader->AllocateAndSetParameters(GraphBuilder, InputRGGBuffer, OutputResource, ElementsPerRow, ProcessedOutputDimension.Y);
-
+			FRGB10BitToRGBA10CS::FParameters* Parameters = ComputeShader->AllocateAndSetParameters(GraphBuilder, InputBuffer, OutputResource, SourceBufferDesc.ElementsPerRow, ProcessedOutputDimension.Y);
 
 			FComputeShaderUtils::AddPass(
 				GraphBuilder
@@ -162,22 +141,11 @@ bool FRivermaxMediaTextureSampleConverter::Convert(FTexture2DRHIRef& InDestinati
 		{
 			TRACE_CPUPROFILER_EVENT_SCOPE(RivermaxSampleConverter::RGB12ShaderSetup);
 
-			const int32 BytesPerElement = sizeof(FRGB12BitToRGBA12CS::FRGB12BitBuffer);
-			const uint32 Stride = SamplePtr->GetStride();
-			const int32 ElementsPerRow = (Stride / BytesPerElement) + ((Stride % BytesPerElement > 0) ? 1 : 0);
-			const int32 ElementCount = ElementsPerRow * InDestinationTexture->GetDesc().Extent.Y;
-
 			FRGB12BitToRGBA12CS::FPermutationDomain PermutationVector;
 			PermutationVector.Set<FRGB12BitToRGBA12CS::FSRGBToLinear>(bDoSRGBToLinear);
 
-			FRDGBufferRef InputRGGBuffer = CreateStructuredBuffer(GraphBuilder, TEXT("RivermaxInputBuffer"), BytesPerElement, ElementCount, SamplePtr->GetBuffer(), BytesPerElement * ElementCount);
-			constexpr int32 PixelsPerInput = 8;
-			const FIntPoint ProcessedOutputDimension = { InDestinationTexture->GetDesc().Extent.X / PixelsPerInput,InDestinationTexture->GetDesc().Extent.Y };
-			GroupCount = FComputeShaderUtils::GetGroupCount(ProcessedOutputDimension, FComputeShaderUtils::kGolden2DGroupSize);
-			FGlobalShaderMap* GlobalShaderMap = GetGlobalShaderMap(GMaxRHIFeatureLevel);
 			TShaderMapRef<FRGB12BitToRGBA12CS> ComputeShader(GlobalShaderMap, PermutationVector);
-			FRGB12BitToRGBA12CS::FParameters* Parameters = ComputeShader->AllocateAndSetParameters(GraphBuilder, InputRGGBuffer, OutputResource, ElementsPerRow, ProcessedOutputDimension.Y);
-
+			FRGB12BitToRGBA12CS::FParameters* Parameters = ComputeShader->AllocateAndSetParameters(GraphBuilder, InputBuffer, OutputResource, SourceBufferDesc.ElementsPerRow, ProcessedOutputDimension.Y);
 
 			FComputeShaderUtils::AddPass(
 				GraphBuilder
@@ -191,22 +159,11 @@ bool FRivermaxMediaTextureSampleConverter::Convert(FTexture2DRHIRef& InDestinati
 		{
 			TRACE_CPUPROFILER_EVENT_SCOPE(RivermaxSampleConverter::RGB16FloatShaderSetup);
 
-			const int32 BytesPerElement = sizeof(FRGB16fBitToRGBA16fCS::FRGB16fBuffer);
-			const uint32 Stride = SamplePtr->GetStride();
-			const int32 ElementsPerRow = (Stride / BytesPerElement) + ((Stride % BytesPerElement > 0) ? 1 : 0);
-			const int32 ElementCount = ElementsPerRow * InDestinationTexture->GetDesc().Extent.Y;
-
 			FRGB16fBitToRGBA16fCS::FPermutationDomain PermutationVector;
 			PermutationVector.Set<FRGB16fBitToRGBA16fCS::FSRGBToLinear>(bDoSRGBToLinear);
 
-			FRDGBufferRef InputRGGBuffer = CreateStructuredBuffer(GraphBuilder, TEXT("RivermaxInputBuffer"), BytesPerElement, ElementCount, SamplePtr->GetBuffer(), BytesPerElement * ElementCount);
-			constexpr int32 PixelsPerInput = 2;
-			const FIntPoint ProcessedOutputDimension = { InDestinationTexture->GetDesc().Extent.X / PixelsPerInput,InDestinationTexture->GetDesc().Extent.Y };
-			GroupCount = FComputeShaderUtils::GetGroupCount(ProcessedOutputDimension, FComputeShaderUtils::kGolden2DGroupSize);
-			FGlobalShaderMap* GlobalShaderMap = GetGlobalShaderMap(GMaxRHIFeatureLevel);
 			TShaderMapRef<FRGB16fBitToRGBA16fCS> ComputeShader(GlobalShaderMap, PermutationVector);
-			FRGB16fBitToRGBA16fCS::FParameters* Parameters = ComputeShader->AllocateAndSetParameters(GraphBuilder, InputRGGBuffer, OutputResource, ElementsPerRow, ProcessedOutputDimension.Y);
-
+			FRGB16fBitToRGBA16fCS::FParameters* Parameters = ComputeShader->AllocateAndSetParameters(GraphBuilder, InputBuffer, OutputResource, SourceBufferDesc.ElementsPerRow, ProcessedOutputDimension.Y);
 
 			FComputeShaderUtils::AddPass(
 				GraphBuilder
diff --git a/Engine/Plugins/VirtualProduction/Rivermax/RivermaxMedia/Source/RivermaxMedia/Private/RivermaxMediaUtils.cpp b/Engine/Plugins/VirtualProduction/Rivermax/RivermaxMedia/Source/RivermaxMedia/Private/RivermaxMediaUtils.cpp
--- a/Engine/Plugins/VirtualProduction/Rivermax/RivermaxMedia/Source/RivermaxMedia/Private/RivermaxMediaUtils.cpp
+++ b/Engine/Plugins/VirtualProduction/Rivermax/RivermaxMedia/Source/RivermaxMedia/Private/RivermaxMediaUtils.cpp
@@ -2,6 +2,8 @@
 
 #include "RivermaxMediaUtils.h"
 
+#include "RivermaxShaders.h"
+
 namespace UE::RivermaxMediaUtils::Private
 {
 	UE::RivermaxCore::ESamplingType MediaOutputPixelFormatToRivermaxSamplingType(ERivermaxMediaOutputPixelFormat InPixelFormat)
@@ -79,4 +81,67 @@ namespace UE::RivermaxMediaUtils::Private
 		}
 		}
 	}
+
+	FSourceBufferDesc GetBufferDescription(const FIntPoint& InResolution, ERivermaxMediaSourcePixelFormat InPixelFormat)
+	{
+		using namespace UE::RivermaxCore;
+		using namespace UE::RivermaxShaders;
+
+		FSourceBufferDesc Description;
+
+		const ESamplingType SamplingType = MediaSourcePixelFormatToRivermaxSamplingType(InPixelFormat);
+		const FVideoFormatInfo Info = FStandardVideoFormat::GetVideoFormatInfo(SamplingType);
+
+		// Compute horizontal byte count (stride) of resolution
+		const uint32 PixelAlignment = Info.PixelGroupCoverage;
+		Description.PixelGroupCount = PixelAlignment > 0 ? InResolution.X / PixelAlignment : 0;
+		Description.BytesPerRow = Description.PixelGroupCount * Info.PixelGroupSize;
+
+		switch (SamplingType)
+		{
+		case ESamplingType::YUV422_8bit:
+		{
+			Description.BytesPerElement = sizeof(FRGBToYUV8Bit422CS::FYUV8Bit422Buffer);
+			break;
+		}
+		case ESamplingType::YUV422_10bit:
+		{
+			Description.BytesPerElement = sizeof(FRGBToYUV10Bit422LittleEndianCS::FYUV10Bit422LEBuffer);
+			break;
+		}
+		case ESamplingType::RGB_8bit:
+		{
+			Description.BytesPerElement = sizeof(FRGBToRGB8BitCS::FRGB8BitBuffer);
+			break;
+		}
+		case ESamplingType::RGB_10bit:
+		{
+			Description.BytesPerElement = sizeof(FRGB10BitToRGBA10CS::FRGB10BitBuffer);
+			break;
+		}
+		case ESamplingType::RGB_12bit:
+		{
+			Description.BytesPerElement = sizeof(FRGBToRGB12BitCS::FRGB12BitBuffer);
+			break;
+		}
+		case ESamplingType::RGB_16bitFloat:
+		{
+			Description.BytesPerElement = sizeof(FRGBToRGB16fCS::FRGB16fBuffer);
+			break;
+		}
+		default:
+		{
+			checkNoEntry(); 
+			return Description;
+		}
+		}
+
+		// Shader encoding might not align with pixel group size so we need to have enough elements to represent the last pixel group
+		Description.ElementsPerRow = Description.BytesPerRow / Description.BytesPerElement;
+		Description.ElementsPerRow += Description.BytesPerRow % Description.BytesPerElement != 0 ? 1 : 0;
+		Description.NumberOfElements = Description.ElementsPerRow * InResolution.Y;
+		
+		return Description;
+	}
+
 }
diff --git a/Engine/Plugins/VirtualProduction/Rivermax/RivermaxMedia/Source/RivermaxMedia/Private/RivermaxMediaUtils.h b/Engine/Plugins/VirtualProduction/Rivermax/RivermaxMedia/Source/RivermaxMedia/Private/RivermaxMediaUtils.h
--- a/Engine/Plugins/VirtualProduction/Rivermax/RivermaxMedia/Source/RivermaxMedia/Private/RivermaxMediaUtils.h
+++ b/Engine/Plugins/VirtualProduction/Rivermax/RivermaxMedia/Source/RivermaxMedia/Private/RivermaxMediaUtils.h
@@ -10,8 +10,17 @@
 
 namespace UE::RivermaxMediaUtils::Private
 {
+	struct FSourceBufferDesc
+	{
+		uint32 PixelGroupCount = 0;
+		uint32 BytesPerElement = 0;
+		uint32 ElementsPerRow = 0;
+		uint32 BytesPerRow = 0;
+		uint32 NumberOfElements = 0;
+	};
 
-UE::RivermaxCore::ESamplingType MediaOutputPixelFormatToRivermaxSamplingType(ERivermaxMediaOutputPixelFormat InPixelFormat);
-UE::RivermaxCore::ESamplingType MediaSourcePixelFormatToRivermaxSamplingType(ERivermaxMediaSourcePixelFormat InPixelFormat);
+	UE::RivermaxCore::ESamplingType MediaOutputPixelFormatToRivermaxSamplingType(ERivermaxMediaOutputPixelFormat InPixelFormat);
+	UE::RivermaxCore::ESamplingType MediaSourcePixelFormatToRivermaxSamplingType(ERivermaxMediaSourcePixelFormat InPixelFormat);
+	FSourceBufferDesc GetBufferDescription(const FIntPoint& Resolution, ERivermaxMediaSourcePixelFormat InPixelFormat);
 
 }
\ No newline at end of file
diff --git a/Engine/Plugins/VirtualProduction/Rivermax/RivermaxMedia/Source/RivermaxMedia/Public/RivermaxMediaSource.h b/Engine/Plugins/VirtualProduction/Rivermax/RivermaxMedia/Source/RivermaxMedia/Public/RivermaxMediaSource.h
--- a/Engine/Plugins/VirtualProduction/Rivermax/RivermaxMedia/Source/RivermaxMedia/Public/RivermaxMediaSource.h
+++ b/Engine/Plugins/VirtualProduction/Rivermax/RivermaxMedia/Source/RivermaxMedia/Public/RivermaxMediaSource.h
@@ -56,6 +56,11 @@ public:
 	UPROPERTY(EditAnywhere, BlueprintReadWrite, Category = "Video")
 	bool bIsSRGBInput = false;
 
+	/** Whether to use GPUDirect if available (Memcopy from NIC to GPU directly bypassing system memory) if available */
+	UPROPERTY(EditAnywhere, BlueprintReadWrite, AdvancedDisplay, Category = "Video")
+	bool bUseGPUDirect = true;
+
+
 public:
 	//~ Begin IMediaOptions interface
 	virtual bool GetMediaOption(const FName& Key, bool DefaultValue) const override;
