Fixing nonunity error.

diff --git a/Engine/Plugins/VirtualProduction/Rivermax/RivermaxCore/Source/RivermaxCore/Private/RivermaxManager.cpp b/Engine/Plugins/VirtualProduction/Rivermax/RivermaxCore/Source/RivermaxCore/Private/RivermaxManager.cpp
--- a/Engine/Plugins/VirtualProduction/Rivermax/RivermaxCore/Source/RivermaxCore/Private/RivermaxManager.cpp
+++ b/Engine/Plugins/VirtualProduction/Rivermax/RivermaxCore/Source/RivermaxCore/Private/RivermaxManager.cpp
@@ -3,6 +3,7 @@
 #include "RivermaxManager.h"
 
 #include "CudaModule.h"
+#include "HAL/IConsoleManager.h"
 #include "Interfaces/IPv4/IPv4Endpoint.h"
 #include "IRivermaxCoreModule.h"
 #include "Misc/App.h"
