Non-comprehensive pass on removing dllexport's from global shader types. Shader types only need exports on their boilerplate macro (DECLARE_GLOBAL_SHADER -> DECLARE_EXPORTED_GLOBAL_SHADER) and any "setter" methods not defined inline.

diff --git a/Engine/Plugins/Compositing/OpenColorIO/Source/OpenColorIO/Public/OpenColorIOShader.h b/Engine/Plugins/Compositing/OpenColorIO/Source/OpenColorIO/Public/OpenColorIOShader.h
--- a/Engine/Plugins/Compositing/OpenColorIO/Source/OpenColorIO/Public/OpenColorIOShader.h
+++ b/Engine/Plugins/Compositing/OpenColorIO/Source/OpenColorIO/Public/OpenColorIOShader.h
@@ -70,10 +70,10 @@ namespace OpenColorIOShader
 	static constexpr uint32 MaximumTextureSlots = 8;
 }
 
-class OPENCOLORIO_API FOpenColorIOPixelShader : public FGlobalShader
+class FOpenColorIOPixelShader : public FGlobalShader
 {
 public:
-	DECLARE_SHADER_TYPE(FOpenColorIOPixelShader, OpenColorIO);
+	DECLARE_EXPORTED_SHADER_TYPE(FOpenColorIOPixelShader, OpenColorIO, OPENCOLORIO_API);
 	SHADER_USE_PARAMETER_STRUCT(FOpenColorIOPixelShader, FGlobalShader);
 
 	using FParameters = FOpenColorIOPixelShaderParameters;
diff --git a/Engine/Plugins/Experimental/LandscapePatch/Source/LandscapePatch/Public/LandscapeCircleHeightPatchPS.h b/Engine/Plugins/Experimental/LandscapePatch/Source/LandscapePatch/Public/LandscapeCircleHeightPatchPS.h
--- a/Engine/Plugins/Experimental/LandscapePatch/Source/LandscapePatch/Public/LandscapeCircleHeightPatchPS.h
+++ b/Engine/Plugins/Experimental/LandscapePatch/Source/LandscapePatch/Public/LandscapeCircleHeightPatchPS.h
@@ -2,16 +2,15 @@
 
 #pragma once
 
-
 #include "GlobalShader.h"
 #include "ShaderParameterStruct.h"
 
 /**
  * Shader that applies a circle height patch to a landscape heightmap.
  */
-class LANDSCAPEPATCH_API FLandscapeCircleHeightPatchPS : public FGlobalShader
+class FLandscapeCircleHeightPatchPS : public FGlobalShader
 {
-	DECLARE_GLOBAL_SHADER(FLandscapeCircleHeightPatchPS);
+	DECLARE_EXPORTED_GLOBAL_SHADER(FLandscapeCircleHeightPatchPS, LANDSCAPEPATCH_API);
 	SHADER_USE_PARAMETER_STRUCT(FLandscapeCircleHeightPatchPS, FGlobalShader);
 
 public:
@@ -21,21 +20,20 @@ public:
 		SHADER_PARAMETER(float, InRadius)
 		SHADER_PARAMETER(float, InFalloff)
 		RENDER_TARGET_BINDING_SLOTS() // Holds our output
-		END_SHADER_PARAMETER_STRUCT()
+	END_SHADER_PARAMETER_STRUCT()
 
 	static bool ShouldCompilePermutation(const FGlobalShaderPermutationParameters& Parameters);
-
 	static void ModifyCompilationEnvironment(const FGlobalShaderPermutationParameters& Parameters, FShaderCompilerEnvironment& OutEnvironment);
 
-	static void AddToRenderGraph(FRDGBuilder& GraphBuilder, FParameters* InParameters);
+	LANDSCAPEPATCH_API static void AddToRenderGraph(FRDGBuilder& GraphBuilder, FParameters* InParameters);
 };
 
 /**
  * Shader that applies a circle patch to a landscape visibility layer.
  */
-class LANDSCAPEPATCH_API FLandscapeCircleVisibilityPatchPS : public FLandscapeCircleHeightPatchPS
+class FLandscapeCircleVisibilityPatchPS : public FLandscapeCircleHeightPatchPS
 {
-	DECLARE_GLOBAL_SHADER(FLandscapeCircleVisibilityPatchPS);
+	DECLARE_EXPORTED_GLOBAL_SHADER(FLandscapeCircleVisibilityPatchPS, LANDSCAPEPATCH_API);
 
 public:
 	FLandscapeCircleVisibilityPatchPS() = default;
@@ -44,7 +42,7 @@ public:
 	{}
 
 	static void ModifyCompilationEnvironment(const FGlobalShaderPermutationParameters& Parameters, FShaderCompilerEnvironment& OutEnvironment);
-	static void AddToRenderGraph(FRDGBuilder& GraphBuilder, FParameters* InParameters);
+	LANDSCAPEPATCH_API static void AddToRenderGraph(FRDGBuilder& GraphBuilder, FParameters* InParameters);
 };
 
 #if UE_ENABLE_INCLUDE_ORDER_DEPRECATED_IN_5_2
diff --git a/Engine/Plugins/Experimental/LandscapePatch/Source/LandscapePatch/Public/LandscapeTexturePatchPS.h b/Engine/Plugins/Experimental/LandscapePatch/Source/LandscapePatch/Public/LandscapeTexturePatchPS.h
--- a/Engine/Plugins/Experimental/LandscapePatch/Source/LandscapePatch/Public/LandscapeTexturePatchPS.h
+++ b/Engine/Plugins/Experimental/LandscapePatch/Source/LandscapePatch/Public/LandscapeTexturePatchPS.h
@@ -30,9 +30,9 @@ namespace UE::Landscape
 /**
  * Shader that applies a texture-based height patch to a landscape heightmap.
  */
-class LANDSCAPEPATCH_API FApplyLandscapeTextureHeightPatchPS : public FGlobalShader
+class FApplyLandscapeTextureHeightPatchPS : public FGlobalShader
 {
-	DECLARE_GLOBAL_SHADER(FApplyLandscapeTextureHeightPatchPS);
+	DECLARE_EXPORTED_GLOBAL_SHADER(FApplyLandscapeTextureHeightPatchPS, LANDSCAPEPATCH_API);
 	SHADER_USE_PARAMETER_STRUCT(FApplyLandscapeTextureHeightPatchPS, FGlobalShader);
 
 public:
@@ -96,11 +96,12 @@ public:
 		SHADER_PARAMETER(uint32, InFlags)
 
 		RENDER_TARGET_BINDING_SLOTS() // Holds our output
-		END_SHADER_PARAMETER_STRUCT()
+	END_SHADER_PARAMETER_STRUCT()
 
 	static bool ShouldCompilePermutation(const FGlobalShaderPermutationParameters& Parameters);
 	static void ModifyCompilationEnvironment(const FGlobalShaderPermutationParameters& InParameters, FShaderCompilerEnvironment& OutEnvironment);
-	static void AddToRenderGraph(FRDGBuilder& GraphBuilder, FParameters* InParameters, const FIntRect& DestinationBounds);
+
+	LANDSCAPEPATCH_API static void AddToRenderGraph(FRDGBuilder& GraphBuilder, FParameters* InParameters, const FIntRect& DestinationBounds);
 };
 
 ENUM_CLASS_FLAGS(FApplyLandscapeTextureHeightPatchPS::EFlags);
@@ -108,9 +109,9 @@ ENUM_CLASS_FLAGS(FApplyLandscapeTextureHeightPatchPS::EFlags);
 /**
  * Simple shader that just offsets each height value in a height patch by a constant.
  */
-class LANDSCAPEPATCH_API FOffsetHeightmapPS : public FGlobalShader
+class FOffsetHeightmapPS : public FGlobalShader
 {
-	DECLARE_GLOBAL_SHADER(FOffsetHeightmapPS);
+	DECLARE_EXPORTED_GLOBAL_SHADER(FOffsetHeightmapPS, LANDSCAPEPATCH_API);
 	SHADER_USE_PARAMETER_STRUCT(FOffsetHeightmapPS, FGlobalShader);
 
 public:
@@ -119,11 +120,12 @@ public:
 		SHADER_PARAMETER_RDG_TEXTURE_SRV(Texture2D<float4>, InHeightmap)
 		SHADER_PARAMETER(float, InHeightOffset)
 		RENDER_TARGET_BINDING_SLOTS() // Holds our output
-		END_SHADER_PARAMETER_STRUCT()
+	END_SHADER_PARAMETER_STRUCT()
 
 	static bool ShouldCompilePermutation(const FGlobalShaderPermutationParameters& Parameters);
 	static void ModifyCompilationEnvironment(const FGlobalShaderPermutationParameters& InParameters, FShaderCompilerEnvironment& OutEnvironment);
-	static void AddToRenderGraph(FRDGBuilder& GraphBuilder, FParameters* InParameters);
+
+	LANDSCAPEPATCH_API static void AddToRenderGraph(FRDGBuilder& GraphBuilder, FParameters* InParameters);
 };
 
 /**
@@ -132,9 +134,9 @@ public:
  * This is not perfectly reversible (in case of clamping and due to rounding), but it lets us store the
  * texture in the way that it would be applied to the landscape (usually).
  */
-class LANDSCAPEPATCH_API FConvertToNativeLandscapePatchPS : public FGlobalShader
+class FConvertToNativeLandscapePatchPS : public FGlobalShader
 {
-	DECLARE_GLOBAL_SHADER(FConvertToNativeLandscapePatchPS);
+	DECLARE_EXPORTED_GLOBAL_SHADER(FConvertToNativeLandscapePatchPS, LANDSCAPEPATCH_API);
 	SHADER_USE_PARAMETER_STRUCT(FConvertToNativeLandscapePatchPS, FGlobalShader);
 
 public:
@@ -145,11 +147,12 @@ public:
 		SHADER_PARAMETER(float, InHeightScale)
 		SHADER_PARAMETER(float, InHeightOffset)
 		RENDER_TARGET_BINDING_SLOTS() // Holds our output
-		END_SHADER_PARAMETER_STRUCT()
+	END_SHADER_PARAMETER_STRUCT()
 
 	static bool ShouldCompilePermutation(const FGlobalShaderPermutationParameters& Parameters);
 	static void ModifyCompilationEnvironment(const FGlobalShaderPermutationParameters& InParameters, FShaderCompilerEnvironment& OutEnvironment);
-	static void AddToRenderGraph(FRDGBuilder& GraphBuilder, FRDGTextureRef SourceTexture,
+
+	LANDSCAPEPATCH_API static void AddToRenderGraph(FRDGBuilder& GraphBuilder, FRDGTextureRef SourceTexture,
 		FRDGTextureRef DestinationTexture, const FLandscapeHeightPatchConvertToNativeParams& Params);
 };
 
@@ -157,9 +160,9 @@ public:
  * Shader that undoes the conversion done by FConvertToNativeLandscapePatchPS (to the extent possible, since
  * rounding/clamping makes it not perfectly recoverable).
  */
-class LANDSCAPEPATCH_API FConvertBackFromNativeLandscapePatchPS : public FGlobalShader
+class FConvertBackFromNativeLandscapePatchPS : public FGlobalShader
 {
-	DECLARE_GLOBAL_SHADER(FConvertBackFromNativeLandscapePatchPS);
+	DECLARE_EXPORTED_GLOBAL_SHADER(FConvertBackFromNativeLandscapePatchPS, LANDSCAPEPATCH_API);
 	SHADER_USE_PARAMETER_STRUCT(FConvertBackFromNativeLandscapePatchPS, FGlobalShader);
 
 public:
@@ -170,17 +173,18 @@ public:
 		SHADER_PARAMETER(float, InHeightScale)
 		SHADER_PARAMETER(float, InHeightOffset)
 		RENDER_TARGET_BINDING_SLOTS() // Holds our output
-		END_SHADER_PARAMETER_STRUCT()
+	END_SHADER_PARAMETER_STRUCT()
 
 	static bool ShouldCompilePermutation(const FGlobalShaderPermutationParameters& Parameters);
 	static void ModifyCompilationEnvironment(const FGlobalShaderPermutationParameters& InParameters, FShaderCompilerEnvironment& OutEnvironment);
-	static void AddToRenderGraph(FRDGBuilder& GraphBuilder, FRDGTextureRef SourceTexture,
+
+	LANDSCAPEPATCH_API static void AddToRenderGraph(FRDGBuilder& GraphBuilder, FRDGTextureRef SourceTexture,
 		FRDGTextureRef DestinationTexture, const FLandscapeHeightPatchConvertToNativeParams& Params);
 };
 
-class LANDSCAPEPATCH_API FApplyLandscapeTextureWeightPatchPS : public FGlobalShader
+class FApplyLandscapeTextureWeightPatchPS : public FGlobalShader
 {
-	DECLARE_GLOBAL_SHADER(FApplyLandscapeTextureWeightPatchPS);
+	DECLARE_EXPORTED_GLOBAL_SHADER(FApplyLandscapeTextureWeightPatchPS, LANDSCAPEPATCH_API);
 	SHADER_USE_PARAMETER_STRUCT(FApplyLandscapeTextureWeightPatchPS, FGlobalShader);
 
 public:
@@ -234,18 +238,19 @@ public:
 		SHADER_PARAMETER(uint32, InFlags)
 
 		RENDER_TARGET_BINDING_SLOTS() // Holds our output
-		END_SHADER_PARAMETER_STRUCT()
+	END_SHADER_PARAMETER_STRUCT()
 
-		static bool ShouldCompilePermutation(const FGlobalShaderPermutationParameters& Parameters);
+	static bool ShouldCompilePermutation(const FGlobalShaderPermutationParameters& Parameters);
 	static void ModifyCompilationEnvironment(const FGlobalShaderPermutationParameters& InParameters, FShaderCompilerEnvironment& OutEnvironment);
-	static void AddToRenderGraph(FRDGBuilder& GraphBuilder, FParameters* InParameters, const FIntRect& DestinationBounds);
+
+	LANDSCAPEPATCH_API static void AddToRenderGraph(FRDGBuilder& GraphBuilder, FParameters* InParameters, const FIntRect& DestinationBounds);
 };
 
 ENUM_CLASS_FLAGS(FApplyLandscapeTextureWeightPatchPS::EFlags);
 
-class LANDSCAPEPATCH_API FReinitializeLandscapePatchPS : public FGlobalShader
+class FReinitializeLandscapePatchPS : public FGlobalShader
 {
-	DECLARE_GLOBAL_SHADER(FReinitializeLandscapePatchPS);
+	DECLARE_EXPORTED_GLOBAL_SHADER(FReinitializeLandscapePatchPS, LANDSCAPEPATCH_API);
 	SHADER_USE_PARAMETER_STRUCT(FReinitializeLandscapePatchPS, FGlobalShader);
 
 public:
@@ -256,11 +261,12 @@ public:
 		SHADER_PARAMETER(FMatrix44f, InPatchToSource)
 
 		RENDER_TARGET_BINDING_SLOTS() // Holds our output
-		END_SHADER_PARAMETER_STRUCT()
+	END_SHADER_PARAMETER_STRUCT()
 
 	static bool ShouldCompilePermutation(const FGlobalShaderPermutationParameters& Parameters);
 	static void ModifyCompilationEnvironment(const FGlobalShaderPermutationParameters& InParameters, FShaderCompilerEnvironment& OutEnvironment);
-	static void AddToRenderGraph(FRDGBuilder& GraphBuilder, FParameters* InParameters);
+
+	LANDSCAPEPATCH_API static void AddToRenderGraph(FRDGBuilder& GraphBuilder, FParameters* InParameters);
 };
 
 /**
@@ -268,9 +274,9 @@ public:
  */
 //~ Theoretically CopyToResolveTarget or AddCopyToResolveTargetPass should work, but I was unable to make them
 //~ work without lots of complaints from the renderer.
-class LANDSCAPEPATCH_API FSimpleTextureCopyPS : public FGlobalShader
+class FSimpleTextureCopyPS : public FGlobalShader
 {
-	DECLARE_GLOBAL_SHADER(FSimpleTextureCopyPS);
+	DECLARE_EXPORTED_GLOBAL_SHADER(FSimpleTextureCopyPS, LANDSCAPEPATCH_API);
 	SHADER_USE_PARAMETER_STRUCT(FSimpleTextureCopyPS, FGlobalShader);
 
 public:
@@ -280,11 +286,12 @@ public:
 		SHADER_PARAMETER_SAMPLER(SamplerState, InSourceSampler)
 		SHADER_PARAMETER(FVector2f, InDestinationResolution)
 		RENDER_TARGET_BINDING_SLOTS() // Holds our output
-		END_SHADER_PARAMETER_STRUCT()
+	END_SHADER_PARAMETER_STRUCT()
 
 	static bool ShouldCompilePermutation(const FGlobalShaderPermutationParameters& Parameters);
 	static void ModifyCompilationEnvironment(const FGlobalShaderPermutationParameters& InParameters, FShaderCompilerEnvironment& OutEnvironment);
-	static void AddToRenderGraph(FRDGBuilder& GraphBuilder, FRDGTextureRef SourceTexture, FRDGTextureRef DestinationTexture);
+
+	LANDSCAPEPATCH_API static void AddToRenderGraph(FRDGBuilder& GraphBuilder, FRDGTextureRef SourceTexture, FRDGTextureRef DestinationTexture);
 };
 
 }//end UE::Landscape
diff --git a/Engine/Plugins/Experimental/NNI/Source/Private/NeuralNetworkInferenceShaders/Internal/BatchNormalizationCS.h b/Engine/Plugins/Experimental/NNI/Source/Private/NeuralNetworkInferenceShaders/Internal/BatchNormalizationCS.h
--- a/Engine/Plugins/Experimental/NNI/Source/Private/NeuralNetworkInferenceShaders/Internal/BatchNormalizationCS.h
+++ b/Engine/Plugins/Experimental/NNI/Source/Private/NeuralNetworkInferenceShaders/Internal/BatchNormalizationCS.h
@@ -10,18 +10,16 @@
 #include "RenderGraphUtils.h"
 #include "ShaderParameterUtils.h"
 
-
-
-class NEURALNETWORKINFERENCESHADERS_API FBatchNormalizationCS : public FGlobalShader
+class FBatchNormalizationCS : public FGlobalShader
 {
-	DECLARE_GLOBAL_SHADER(FBatchNormalizationCS);
+	DECLARE_EXPORTED_GLOBAL_SHADER(FBatchNormalizationCS, NEURALNETWORKINFERENCESHADERS_API);
 	SHADER_USE_PARAMETER_STRUCT(FBatchNormalizationCS, FGlobalShader)
 
 	class FBatchNormalizationMode : SHADER_PERMUTATION_ENUM_CLASS("BATCH_NORMALIZATION_MODE", EBatchNormalizationMode);
 	class FIsInlined : SHADER_PERMUTATION_BOOL("IS_INLINED");
 	using FPermutationDomain = TShaderPermutationDomain<FBatchNormalizationMode, FIsInlined>;
 
-	static const uint32 THREADGROUP_SIZE_X;
+	NEURALNETWORKINFERENCESHADERS_API static const uint32 THREADGROUP_SIZE_X;
 
 	static void ModifyCompilationEnvironment(const FGlobalShaderPermutationParameters& InParameters, FShaderCompilerEnvironment& OutEnvironment);
 
diff --git a/Engine/Plugins/Experimental/NNI/Source/Private/NeuralNetworkInferenceShaders/Internal/ConvBaseCS.h b/Engine/Plugins/Experimental/NNI/Source/Private/NeuralNetworkInferenceShaders/Internal/ConvBaseCS.h
--- a/Engine/Plugins/Experimental/NNI/Source/Private/NeuralNetworkInferenceShaders/Internal/ConvBaseCS.h
+++ b/Engine/Plugins/Experimental/NNI/Source/Private/NeuralNetworkInferenceShaders/Internal/ConvBaseCS.h
@@ -10,19 +10,17 @@
 #include "RenderGraphUtils.h"
 #include "ShaderParameterUtils.h"
 
-
-
-class NEURALNETWORKINFERENCESHADERS_API FConvBaseCS : public FGlobalShader
+class FConvBaseCS : public FGlobalShader
 {
-	DECLARE_GLOBAL_SHADER(FConvBaseCS);
+	DECLARE_EXPORTED_GLOBAL_SHADER(FConvBaseCS, NEURALNETWORKINFERENCESHADERS_API);
 	SHADER_USE_PARAMETER_STRUCT(FConvBaseCS, FGlobalShader)
 
 	class FHasBias : SHADER_PERMUTATION_BOOL("HAS_BIAS");
 	class FConvMode : SHADER_PERMUTATION_ENUM_CLASS("CONV_MODE", EConvMode);
 	using FPermutationDomain = TShaderPermutationDomain<FHasBias, FConvMode>;
 
-	static const uint32 THREADGROUP_SIZE_X;
-	static const uint32 MAX_NUMBER_DIMENSIONS;
+	NEURALNETWORKINFERENCESHADERS_API static const uint32 THREADGROUP_SIZE_X;
+	NEURALNETWORKINFERENCESHADERS_API static const uint32 MAX_NUMBER_DIMENSIONS;
 
 	static void ModifyCompilationEnvironment(const FGlobalShaderPermutationParameters& InParameters, FShaderCompilerEnvironment& OutEnvironment);
 
diff --git a/Engine/Plugins/Experimental/NNI/Source/Private/NeuralNetworkInferenceShaders/Internal/ConvTransposeCS.h b/Engine/Plugins/Experimental/NNI/Source/Private/NeuralNetworkInferenceShaders/Internal/ConvTransposeCS.h
--- a/Engine/Plugins/Experimental/NNI/Source/Private/NeuralNetworkInferenceShaders/Internal/ConvTransposeCS.h
+++ b/Engine/Plugins/Experimental/NNI/Source/Private/NeuralNetworkInferenceShaders/Internal/ConvTransposeCS.h
@@ -10,17 +10,15 @@
 #include "RenderGraphUtils.h"
 #include "ShaderParameterUtils.h"
 
-
-
-class NEURALNETWORKINFERENCESHADERS_API FConvTransposeCS : public FGlobalShader
+class FConvTransposeCS : public FGlobalShader
 {
-	DECLARE_GLOBAL_SHADER(FConvTransposeCS);
+	DECLARE_EXPORTED_GLOBAL_SHADER(FConvTransposeCS, NEURALNETWORKINFERENCESHADERS_API);
 	SHADER_USE_PARAMETER_STRUCT(FConvTransposeCS, FGlobalShader)
 
 	class FConvMode : SHADER_PERMUTATION_ENUM_CLASS("CONV_MODE", EConvMode);
 	using FPermutationDomain = TShaderPermutationDomain<FConvMode>;
 
-	static const uint32 THREADGROUP_SIZE_X;
+	NEURALNETWORKINFERENCESHADERS_API static const uint32 THREADGROUP_SIZE_X;
 
 	static void ModifyCompilationEnvironment(const FGlobalShaderPermutationParameters& InParameters, FShaderCompilerEnvironment& OutEnvironment);
 
diff --git a/Engine/Plugins/Experimental/NNI/Source/Private/NeuralNetworkInferenceShaders/Internal/CopyCS.h b/Engine/Plugins/Experimental/NNI/Source/Private/NeuralNetworkInferenceShaders/Internal/CopyCS.h
--- a/Engine/Plugins/Experimental/NNI/Source/Private/NeuralNetworkInferenceShaders/Internal/CopyCS.h
+++ b/Engine/Plugins/Experimental/NNI/Source/Private/NeuralNetworkInferenceShaders/Internal/CopyCS.h
@@ -10,14 +10,12 @@
 #include "RenderGraphUtils.h"
 #include "ShaderParameterUtils.h"
 
-
-
-class NEURALNETWORKINFERENCESHADERS_API FCopyCS : public FGlobalShader
+class FCopyCS : public FGlobalShader
 {
-	DECLARE_GLOBAL_SHADER(FCopyCS);
+	DECLARE_EXPORTED_GLOBAL_SHADER(FCopyCS, NEURALNETWORKINFERENCESHADERS_API);
 	SHADER_USE_PARAMETER_STRUCT(FCopyCS, FGlobalShader)
 
-	static const uint32 THREADGROUP_SIZE_X;
+	NEURALNETWORKINFERENCESHADERS_API static const uint32 THREADGROUP_SIZE_X;
 
 	static void ModifyCompilationEnvironment(const FGlobalShaderPermutationParameters& InParameters, FShaderCompilerEnvironment& OutEnvironment);
 
diff --git a/Engine/Plugins/Experimental/NNI/Source/Private/NeuralNetworkInferenceShaders/Internal/ElementWiseCS.h b/Engine/Plugins/Experimental/NNI/Source/Private/NeuralNetworkInferenceShaders/Internal/ElementWiseCS.h
--- a/Engine/Plugins/Experimental/NNI/Source/Private/NeuralNetworkInferenceShaders/Internal/ElementWiseCS.h
+++ b/Engine/Plugins/Experimental/NNI/Source/Private/NeuralNetworkInferenceShaders/Internal/ElementWiseCS.h
@@ -10,14 +10,12 @@
 #include "RenderGraphUtils.h"
 #include "ShaderParameterUtils.h"
 
-
-
-class NEURALNETWORKINFERENCESHADERS_API FElementWiseCS : public FGlobalShader
+class FElementWiseCS : public FGlobalShader
 {
-	DECLARE_GLOBAL_SHADER(FElementWiseCS);
+	DECLARE_EXPORTED_GLOBAL_SHADER(FElementWiseCS, NEURALNETWORKINFERENCESHADERS_API);
 	SHADER_USE_PARAMETER_STRUCT(FElementWiseCS, FGlobalShader)
 
-	static const uint32 THREADGROUP_SIZE_X;
+	NEURALNETWORKINFERENCESHADERS_API static const uint32 THREADGROUP_SIZE_X;
 
 	class FShaderType : SHADER_PERMUTATION_ENUM_CLASS("SHADER_FUNCTION", EElementWiseOperator);
 	class FIsInlined : SHADER_PERMUTATION_BOOL("IS_INLINED");
diff --git a/Engine/Plugins/Experimental/NNI/Source/Private/NeuralNetworkInferenceShaders/Internal/GemmCS.h b/Engine/Plugins/Experimental/NNI/Source/Private/NeuralNetworkInferenceShaders/Internal/GemmCS.h
--- a/Engine/Plugins/Experimental/NNI/Source/Private/NeuralNetworkInferenceShaders/Internal/GemmCS.h
+++ b/Engine/Plugins/Experimental/NNI/Source/Private/NeuralNetworkInferenceShaders/Internal/GemmCS.h
@@ -10,18 +10,16 @@
 #include "RenderGraphUtils.h"
 #include "ShaderParameterUtils.h"
 
-
-
-class NEURALNETWORKINFERENCESHADERS_API FGemmCS : public FGlobalShader
+class FGemmCS : public FGlobalShader
 {
-	DECLARE_GLOBAL_SHADER(FGemmCS);
+	DECLARE_EXPORTED_GLOBAL_SHADER(FGemmCS, NEURALNETWORKINFERENCESHADERS_API);
 	SHADER_USE_PARAMETER_STRUCT(FGemmCS, FGlobalShader)
 
 	class FGemmMode : SHADER_PERMUTATION_ENUM_CLASS("GEMM_MODE", EGemmMode);
 	using FPermutationDomain = TShaderPermutationDomain<FGemmMode>;
 
-	static const uint32 THREADGROUP_SIZE_X;
-	static const uint32 THREADGROUP_SIZE_Y;
+	NEURALNETWORKINFERENCESHADERS_API static const uint32 THREADGROUP_SIZE_X;
+	NEURALNETWORKINFERENCESHADERS_API static const uint32 THREADGROUP_SIZE_Y;
 
 	static void ModifyCompilationEnvironment(const FGlobalShaderPermutationParameters& InParameters, FShaderCompilerEnvironment& OutEnvironment);
 
diff --git a/Engine/Plugins/Experimental/NNI/Source/Private/NeuralNetworkInferenceShaders/Internal/MultidirectionalBroadcastCS.h b/Engine/Plugins/Experimental/NNI/Source/Private/NeuralNetworkInferenceShaders/Internal/MultidirectionalBroadcastCS.h
--- a/Engine/Plugins/Experimental/NNI/Source/Private/NeuralNetworkInferenceShaders/Internal/MultidirectionalBroadcastCS.h
+++ b/Engine/Plugins/Experimental/NNI/Source/Private/NeuralNetworkInferenceShaders/Internal/MultidirectionalBroadcastCS.h
@@ -10,15 +10,13 @@
 #include "RenderGraphUtils.h"
 #include "ShaderParameterUtils.h"
 
-
-
-class NEURALNETWORKINFERENCESHADERS_API FMultidirectionalBroadcastCS : public FGlobalShader
+class FMultidirectionalBroadcastCS : public FGlobalShader
 {
-	DECLARE_GLOBAL_SHADER(FMultidirectionalBroadcastCS);
+	DECLARE_EXPORTED_GLOBAL_SHADER(FMultidirectionalBroadcastCS, NEURALNETWORKINFERENCESHADERS_API);
 	SHADER_USE_PARAMETER_STRUCT(FMultidirectionalBroadcastCS, FGlobalShader)
 
-	static const uint32 THREADGROUP_SIZE_X;
-	static const uint32 MAX_NUMBER_DIMENSIONS;
+	NEURALNETWORKINFERENCESHADERS_API static const uint32 THREADGROUP_SIZE_X;
+	NEURALNETWORKINFERENCESHADERS_API static const uint32 MAX_NUMBER_DIMENSIONS;
 
 	class FShaderType : SHADER_PERMUTATION_ENUM_CLASS("SHADER_FUNCTION", EMultidirectionalBroadcastOperator);
 	class FInlinedMode : SHADER_PERMUTATION_ENUM_CLASS("INLINED_MODE", EMultidirectionalBroadcastInlinedMode);
diff --git a/Engine/Plugins/FX/Niagara/Source/NiagaraShader/Private/NiagaraDebugShaders.cpp b/Engine/Plugins/FX/Niagara/Source/NiagaraShader/Private/NiagaraDebugShaders.cpp
--- a/Engine/Plugins/FX/Niagara/Source/NiagaraShader/Private/NiagaraDebugShaders.cpp
+++ b/Engine/Plugins/FX/Niagara/Source/NiagaraShader/Private/NiagaraDebugShaders.cpp
@@ -41,7 +41,7 @@ static FAutoConsoleVariableRef CVarNiagaraGpuComputeDebug_OccludedLineColorScale
 
 //////////////////////////////////////////////////////////////////////////
 
-class NIAGARASHADER_API FNiagaraVisualizeTexturePS : public FGlobalShader
+class FNiagaraVisualizeTexturePS : public FGlobalShader
 {
 	DECLARE_GLOBAL_SHADER(FNiagaraVisualizeTexturePS);
 	SHADER_USE_PARAMETER_STRUCT(FNiagaraVisualizeTexturePS, FGlobalShader);
@@ -81,7 +81,7 @@ IMPLEMENT_GLOBAL_SHADER(FNiagaraVisualizeTexturePS, "/Plugin/FX/Niagara/Private/
 
 //////////////////////////////////////////////////////////////////////////
 
-class NIAGARASHADER_API FNiagaraClearUAVCS : public FGlobalShader
+class FNiagaraClearUAVCS : public FGlobalShader
 {
 	DECLARE_GLOBAL_SHADER(FNiagaraClearUAVCS);
 	SHADER_USE_PARAMETER_STRUCT(FNiagaraClearUAVCS, FGlobalShader);
@@ -106,7 +106,7 @@ IMPLEMENT_GLOBAL_SHADER(FNiagaraClearUAVCS, "/Plugin/FX/Niagara/Private/NiagaraD
 
 //////////////////////////////////////////////////////////////////////////
 
-class NIAGARASHADER_API FNiagaraDebugDrawLineVS : public FGlobalShader
+class FNiagaraDebugDrawLineVS : public FGlobalShader
 {
 	DECLARE_GLOBAL_SHADER(FNiagaraDebugDrawLineVS);
 	SHADER_USE_PARAMETER_STRUCT(FNiagaraDebugDrawLineVS, FGlobalShader);
@@ -128,7 +128,7 @@ class NIAGARASHADER_API FNiagaraDebugDrawLineVS : public FGlobalShader
 	END_SHADER_PARAMETER_STRUCT()
 };
 
-class NIAGARASHADER_API FNiagaraDebugDrawLinePS : public FGlobalShader
+class FNiagaraDebugDrawLinePS : public FGlobalShader
 {
 	DECLARE_GLOBAL_SHADER(FNiagaraDebugDrawLinePS);
 	SHADER_USE_PARAMETER_STRUCT(FNiagaraDebugDrawLinePS, FGlobalShader);
diff --git a/Engine/Plugins/FX/Niagara/Source/NiagaraShader/Public/NiagaraRibbonCompute.h b/Engine/Plugins/FX/Niagara/Source/NiagaraShader/Public/NiagaraRibbonCompute.h
--- a/Engine/Plugins/FX/Niagara/Source/NiagaraShader/Public/NiagaraRibbonCompute.h
+++ b/Engine/Plugins/FX/Niagara/Source/NiagaraShader/Public/NiagaraRibbonCompute.h
@@ -82,9 +82,6 @@ struct FNiagaraRibbonComputeCommon
 	static constexpr uint32 IndexGenOptimalLoopVertexLimit = 32;
 };
 
-
-
-
 BEGIN_SHADER_PARAMETER_STRUCT(FRibbonOrderSortParameters, NIAGARASHADER_API)
 	SHADER_PARAMETER_STRUCT_INCLUDE(FRibbonComputeUniformParameters, Common)
 	SHADER_PARAMETER_SRV(Buffer<uint32>, SortedIndices)
@@ -96,9 +93,9 @@ END_SHADER_PARAMETER_STRUCT()
 /**
 * Compute shader used to generate particle sort keys.
 */
-class NIAGARASHADER_API FNiagaraRibbonSortPhase1CS : public FGlobalShader
+class FNiagaraRibbonSortPhase1CS : public FGlobalShader
 {
-	DECLARE_GLOBAL_SHADER(FNiagaraRibbonSortPhase1CS);
+	DECLARE_EXPORTED_GLOBAL_SHADER(FNiagaraRibbonSortPhase1CS, NIAGARASHADER_API);
 	SHADER_USE_PARAMETER_STRUCT(FNiagaraRibbonSortPhase1CS, FGlobalShader);
 
 	static constexpr uint32 BubbleSortGroupWidth = 32;
@@ -112,9 +109,9 @@ class NIAGARASHADER_API FNiagaraRibbonSortPhase1CS : public FGlobalShader
 /**
 * Compute shader used to generate particle sort keys.
 */
-class NIAGARASHADER_API FNiagaraRibbonSortPhase2CS : public FGlobalShader
+class FNiagaraRibbonSortPhase2CS : public FGlobalShader
 {
-	DECLARE_GLOBAL_SHADER(FNiagaraRibbonSortPhase2CS);
+	DECLARE_EXPORTED_GLOBAL_SHADER(FNiagaraRibbonSortPhase2CS, NIAGARASHADER_API);
 	SHADER_USE_PARAMETER_STRUCT(FNiagaraRibbonSortPhase2CS, FGlobalShader);
 
 	static constexpr uint32 ThreadGroupSize = 64;
@@ -125,7 +122,6 @@ class NIAGARASHADER_API FNiagaraRibbonSortPhase2CS : public FGlobalShader
 	static void ModifyCompilationEnvironment(const FGlobalShaderPermutationParameters& Parameters, FShaderCompilerEnvironment& OutEnvironment);
 };
 
-
 BEGIN_SHADER_PARAMETER_STRUCT(FNiagaraRibbonVertexReductionParameters, NIAGARASHADER_API)
 	SHADER_PARAMETER_STRUCT_INCLUDE(FRibbonComputeUniformParameters, Common)
 	SHADER_PARAMETER_SRV(Buffer<uint32>, SortedIndices)
@@ -141,9 +137,9 @@ BEGIN_SHADER_PARAMETER_STRUCT(FNiagaraRibbonVertexReductionParameters, NIAGARASH
 	SHADER_PARAMETER(int, PrefixScanStride)
 END_SHADER_PARAMETER_STRUCT()
 
-class NIAGARASHADER_API FNiagaraRibbonVertexReductionInitializationCS : public FGlobalShader
+class FNiagaraRibbonVertexReductionInitializationCS : public FGlobalShader
 {
-	DECLARE_GLOBAL_SHADER(FNiagaraRibbonVertexReductionInitializationCS);
+	DECLARE_EXPORTED_GLOBAL_SHADER(FNiagaraRibbonVertexReductionInitializationCS, NIAGARASHADER_API);
 	SHADER_USE_PARAMETER_STRUCT(FNiagaraRibbonVertexReductionInitializationCS, FGlobalShader);
 	
 	using FPermutationDomain = TShaderPermutationDomain<FRibbonHasFullRibbonID, FRibbonHasRibbonID, FRibbonWantsConstantTessellation, FRibbonWantsAutomaticTessellation, FRibbonHasTwist>;
@@ -152,9 +148,9 @@ class NIAGARASHADER_API FNiagaraRibbonVertexReductionInitializationCS : public F
 	static void ModifyCompilationEnvironment(const FGlobalShaderPermutationParameters& Parameters, FShaderCompilerEnvironment& OutEnvironment);
 };
 
-class NIAGARASHADER_API FNiagaraRibbonVertexReductionPropagateCS : public FGlobalShader
+class FNiagaraRibbonVertexReductionPropagateCS : public FGlobalShader
 {
-	DECLARE_GLOBAL_SHADER(FNiagaraRibbonVertexReductionPropagateCS);
+	DECLARE_EXPORTED_GLOBAL_SHADER(FNiagaraRibbonVertexReductionPropagateCS, NIAGARASHADER_API);
 	SHADER_USE_PARAMETER_STRUCT(FNiagaraRibbonVertexReductionPropagateCS, FGlobalShader);
 	
 	using FPermutationDomain = TShaderPermutationDomain<FRibbonHasFullRibbonID, FRibbonHasRibbonID, FRibbonWantsConstantTessellation, FRibbonWantsAutomaticTessellation, FRibbonHasTwist>;
@@ -176,9 +172,9 @@ BEGIN_SHADER_PARAMETER_STRUCT(FNiagaraRibbonVertexReductionFinalizationParameter
 	SHADER_PARAMETER(int, FinalizationThreadBlockSize)
 END_SHADER_PARAMETER_STRUCT()
 
-class NIAGARASHADER_API FNiagaraRibbonVertexReductionFinalizeCS : public FGlobalShader
+class FNiagaraRibbonVertexReductionFinalizeCS : public FGlobalShader
 {
-	DECLARE_GLOBAL_SHADER(FNiagaraRibbonVertexReductionFinalizeCS);
+	DECLARE_EXPORTED_GLOBAL_SHADER(FNiagaraRibbonVertexReductionFinalizeCS, NIAGARASHADER_API);
 	SHADER_USE_PARAMETER_STRUCT(FNiagaraRibbonVertexReductionFinalizeCS, FGlobalShader);
 	
 	using FPermutationDomain = TShaderPermutationDomain<FRibbonHasFullRibbonID, FRibbonHasRibbonID, FRibbonWantsConstantTessellation, FRibbonWantsAutomaticTessellation, FRibbonHasTwist>;
@@ -210,9 +206,9 @@ BEGIN_SHADER_PARAMETER_STRUCT(FNiagaraRibbonVertexFinalizationParameters, NIAGAR
 	SHADER_PARAMETER(int32, TotalNumRibbons)
 END_SHADER_PARAMETER_STRUCT()
 
-class NIAGARASHADER_API FNiagaraRibbonUVParamCalculationCS : public FGlobalShader
+class FNiagaraRibbonUVParamCalculationCS : public FGlobalShader
 {
-	DECLARE_GLOBAL_SHADER(FNiagaraRibbonUVParamCalculationCS);
+	DECLARE_EXPORTED_GLOBAL_SHADER(FNiagaraRibbonUVParamCalculationCS, NIAGARASHADER_API);
 	SHADER_USE_PARAMETER_STRUCT(FNiagaraRibbonUVParamCalculationCS, FGlobalShader);
 	
 	using FPermutationDomain = TShaderPermutationDomain<FRibbonHasFullRibbonID, FRibbonHasRibbonID, FRibbonWantsConstantTessellation, FRibbonWantsAutomaticTessellation>;
@@ -249,11 +245,10 @@ BEGIN_SHADER_PARAMETER_STRUCT(FNiagaraRibbonInitializeIndices, NIAGARASHADER_API
 	SHADER_PARAMETER(float, GNiagaraRibbonTessellationMinDisplacementError)
 END_SHADER_PARAMETER_STRUCT()
 
-class NIAGARASHADER_API FNiagaraRibbonCreateIndexBufferParamsCS : public FGlobalShader
+class FNiagaraRibbonCreateIndexBufferParamsCS : public FGlobalShader
 {
-	DECLARE_GLOBAL_SHADER(FNiagaraRibbonCreateIndexBufferParamsCS);
+	DECLARE_EXPORTED_GLOBAL_SHADER(FNiagaraRibbonCreateIndexBufferParamsCS, NIAGARASHADER_API);
 	SHADER_USE_PARAMETER_STRUCT(FNiagaraRibbonCreateIndexBufferParamsCS, FGlobalShader);
-
 	
 	using FPermutationDomain = TShaderPermutationDomain<FRibbonWantsConstantTessellation, FRibbonWantsAutomaticTessellation>;
 	using FParameters = FNiagaraRibbonInitializeIndices;
@@ -290,11 +285,10 @@ BEGIN_SHADER_PARAMETER_STRUCT(FNiagaraRibbonGenerateIndices, NIAGARASHADER_API)
 	SHADER_PARAMETER(uint32, SubSegmentBitMask)
 END_SHADER_PARAMETER_STRUCT()
 
-class NIAGARASHADER_API FNiagaraRibbonCreateIndexBufferCS : public FGlobalShader
+class FNiagaraRibbonCreateIndexBufferCS : public FGlobalShader
 {
-	DECLARE_GLOBAL_SHADER(FNiagaraRibbonCreateIndexBufferCS);
+	DECLARE_EXPORTED_GLOBAL_SHADER(FNiagaraRibbonCreateIndexBufferCS, NIAGARASHADER_API);
 	SHADER_USE_PARAMETER_STRUCT(FNiagaraRibbonCreateIndexBufferCS, FGlobalShader);
-
 	
 	using FPermutationDomain = TShaderPermutationDomain<FRibbonHasFullRibbonID, FRibbonHasRibbonID, FRibbonHasHighSliceComplexity>;;
 	using FParameters = FNiagaraRibbonGenerateIndices;
diff --git a/Engine/Plugins/FX/Niagara/Source/NiagaraVertexFactories/Public/NiagaraDrawIndirect.h b/Engine/Plugins/FX/Niagara/Source/NiagaraVertexFactories/Public/NiagaraDrawIndirect.h
--- a/Engine/Plugins/FX/Niagara/Source/NiagaraVertexFactories/Public/NiagaraDrawIndirect.h
+++ b/Engine/Plugins/FX/Niagara/Source/NiagaraVertexFactories/Public/NiagaraDrawIndirect.h
@@ -63,9 +63,9 @@ struct FNiagaraDrawIndirectArgGenTaskInfo
  * Compute shader used to generate GPU emitter draw indirect args.
  * It also resets unused instance count entries.
  */
-class NIAGARAVERTEXFACTORIES_API FNiagaraDrawIndirectArgsGenCS : public FGlobalShader
+class FNiagaraDrawIndirectArgsGenCS : public FGlobalShader
 {
-	DECLARE_GLOBAL_SHADER(FNiagaraDrawIndirectArgsGenCS);
+	DECLARE_EXPORTED_GLOBAL_SHADER(FNiagaraDrawIndirectArgsGenCS, NIAGARAVERTEXFACTORIES_API);
 	SHADER_USE_PARAMETER_STRUCT(FNiagaraDrawIndirectArgsGenCS, FGlobalShader);
 
 public:
@@ -88,9 +88,9 @@ public:
 /**
  * Compute shader used to reset unused instance count entries. Used if the platform doesn't support RW texture buffers
  */
-class NIAGARAVERTEXFACTORIES_API FNiagaraDrawIndirectResetCountsCS : public FGlobalShader
+class FNiagaraDrawIndirectResetCountsCS : public FGlobalShader
 {
-	DECLARE_GLOBAL_SHADER(FNiagaraDrawIndirectResetCountsCS);
+	DECLARE_EXPORTED_GLOBAL_SHADER(FNiagaraDrawIndirectResetCountsCS, NIAGARAVERTEXFACTORIES_API);
 	SHADER_USE_PARAMETER_STRUCT(FNiagaraDrawIndirectResetCountsCS, FGlobalShader);
 
 public:
@@ -103,4 +103,4 @@ public:
 	END_SHADER_PARAMETER_STRUCT()
 
 	static void ModifyCompilationEnvironment(const FGlobalShaderPermutationParameters& Parameters, FShaderCompilerEnvironment& OutEnvironment);
-};
+};
\ No newline at end of file
diff --git a/Engine/Plugins/FX/Niagara/Source/NiagaraVertexFactories/Public/NiagaraGPURayTracingTransformsShader.h b/Engine/Plugins/FX/Niagara/Source/NiagaraVertexFactories/Public/NiagaraGPURayTracingTransformsShader.h
--- a/Engine/Plugins/FX/Niagara/Source/NiagaraVertexFactories/Public/NiagaraGPURayTracingTransformsShader.h
+++ b/Engine/Plugins/FX/Niagara/Source/NiagaraVertexFactories/Public/NiagaraGPURayTracingTransformsShader.h
@@ -14,9 +14,9 @@ NiagaraGPURayTracingTransformsShader.h: Niagara shader to generate the ray traci
 /**
  * Compute shader used to pass GPU instances transforms to the ray tracing TLAS.
  */
-class NIAGARAVERTEXFACTORIES_API FNiagaraGPURayTracingTransformsCS : public FGlobalShader
+class FNiagaraGPURayTracingTransformsCS : public FGlobalShader
 {
-	DECLARE_GLOBAL_SHADER(FNiagaraGPURayTracingTransformsCS);
+	DECLARE_EXPORTED_GLOBAL_SHADER(FNiagaraGPURayTracingTransformsCS, NIAGARAVERTEXFACTORIES_API);
 	SHADER_USE_PARAMETER_STRUCT(FNiagaraGPURayTracingTransformsCS, FGlobalShader);
 
 	static constexpr uint32 ThreadGroupSize = 64;
diff --git a/Engine/Plugins/FX/Niagara/Source/NiagaraVertexFactories/Public/NiagaraSortingGPU.h b/Engine/Plugins/FX/Niagara/Source/NiagaraVertexFactories/Public/NiagaraSortingGPU.h
--- a/Engine/Plugins/FX/Niagara/Source/NiagaraVertexFactories/Public/NiagaraSortingGPU.h
+++ b/Engine/Plugins/FX/Niagara/Source/NiagaraVertexFactories/Public/NiagaraSortingGPU.h
@@ -24,9 +24,9 @@ extern NIAGARAVERTEXFACTORIES_API int32 GNiagaraGPUCullingCPUToGPUThreshold;
 /**
  * Compute shader used to generate particle sort keys.
  */
-class NIAGARAVERTEXFACTORIES_API FNiagaraSortKeyGenCS : public FGlobalShader
+class FNiagaraSortKeyGenCS : public FGlobalShader
 {
-	DECLARE_GLOBAL_SHADER(FNiagaraSortKeyGenCS);
+	DECLARE_EXPORTED_GLOBAL_SHADER(FNiagaraSortKeyGenCS, NIAGARAVERTEXFACTORIES_API);
 	SHADER_USE_PARAMETER_STRUCT(FNiagaraSortKeyGenCS, FGlobalShader);
 
 public:
@@ -79,5 +79,5 @@ public:
 
 	static void ModifyCompilationEnvironment(const FGlobalShaderPermutationParameters& Parameters, FShaderCompilerEnvironment& OutEnvironment);
 
-	static bool UseWaveOps(EShaderPlatform ShaderPlatform);
+	NIAGARAVERTEXFACTORIES_API static bool UseWaveOps(EShaderPlatform ShaderPlatform);
 };
diff --git a/Engine/Plugins/Media/ImgMedia/Source/ExrReaderGpu/Public/ExrSwizzlingShader.h b/Engine/Plugins/Media/ImgMedia/Source/ExrReaderGpu/Public/ExrSwizzlingShader.h
--- a/Engine/Plugins/Media/ImgMedia/Source/ExrReaderGpu/Public/ExrSwizzlingShader.h
+++ b/Engine/Plugins/Media/ImgMedia/Source/ExrReaderGpu/Public/ExrSwizzlingShader.h
@@ -10,15 +10,10 @@
 #include "ShaderParameterStruct.h"
 
 // The vertex shader used by DrawScreenPass to draw a rectangle.
-class EXRREADERGPU_API FExrSwizzleVS : public FGlobalShader
+class FExrSwizzleVS : public FGlobalShader
 {
 public:
-	DECLARE_GLOBAL_SHADER(FExrSwizzleVS);
-
-	static bool ShouldCompilePermutation(const FGlobalShaderPermutationParameters&)
-	{
-		return true;
-	}
+	DECLARE_EXPORTED_GLOBAL_SHADER(FExrSwizzleVS, EXRREADERGPU_API);
 
 	FExrSwizzleVS() = default;
 	FExrSwizzleVS(const ShaderMetaType::CompiledShaderInitializerType& Initializer)
@@ -27,9 +22,9 @@ public:
 };
 
 /** Pixel shader swizzle RGB planar buffer data into proper RGBA texture. */
-class EXRREADERGPU_API FExrSwizzlePS : public FGlobalShader
+class FExrSwizzlePS : public FGlobalShader
 {
-	DECLARE_GLOBAL_SHADER(FExrSwizzlePS);
+	DECLARE_EXPORTED_GLOBAL_SHADER(FExrSwizzlePS, EXRREADERGPU_API);
 	SHADER_USE_PARAMETER_STRUCT(FExrSwizzlePS, FGlobalShader);
 
 	/** If the provided buffer is RGBA the shader would work slightly differently to RGB. */
@@ -39,11 +34,6 @@ class EXRREADERGPU_API FExrSwizzlePS : public FGlobalShader
 	class FPartialTiles : SHADER_PERMUTATION_BOOL("PARTIAL_TILES");
 	using FPermutationDomain = TShaderPermutationDomain<FRgbaSwizzle, FRenderTiles, FCustomExr, FPartialTiles>;
 
-	static bool ShouldCompilePermutation(const FGlobalShaderPermutationParameters& Parameters) 
-	{
-		return true;
-	}
-
 	BEGIN_SHADER_PARAMETER_STRUCT(FParameters, )
 		SHADER_PARAMETER_SRV(StructuredBuffer<uint>, UnswizzledBuffer)
 		SHADER_PARAMETER_SRV(StructuredBuffer<FTileDesc>, TileDescBuffer)
diff --git a/Engine/Plugins/VirtualProduction/Rivermax/RivermaxCore/Source/RivermaxRendering/Public/RivermaxShaders.h b/Engine/Plugins/VirtualProduction/Rivermax/RivermaxCore/Source/RivermaxRendering/Public/RivermaxShaders.h
--- a/Engine/Plugins/VirtualProduction/Rivermax/RivermaxCore/Source/RivermaxRendering/Public/RivermaxShaders.h
+++ b/Engine/Plugins/VirtualProduction/Rivermax/RivermaxCore/Source/RivermaxRendering/Public/RivermaxShaders.h
@@ -17,7 +17,7 @@ namespace UE::RivermaxShaders
 	/**
 	 * Compute shader to convert 2110 YUV8 to RGBA
 	 */
-	class RIVERMAXRENDERING_API FYUV8Bit422ToRGBACS : public FGlobalShader
+	class FYUV8Bit422ToRGBACS : public FGlobalShader
 	{
 	public:
 
@@ -27,7 +27,7 @@ namespace UE::RivermaxShaders
 			uint32 DWord0;
 		};
 
-		DECLARE_GLOBAL_SHADER(FYUV8Bit422ToRGBACS);
+		DECLARE_EXPORTED_GLOBAL_SHADER(FYUV8Bit422ToRGBACS, RIVERMAXRENDERING_API);
 		SHADER_USE_PARAMETER_STRUCT(FYUV8Bit422ToRGBACS, FGlobalShader);
 
 		class FSRGBToLinear : SHADER_PERMUTATION_BOOL("DO_SRGB_TO_LINEAR");
@@ -39,7 +39,7 @@ namespace UE::RivermaxShaders
 			SHADER_PARAMETER(uint32, HorizontalElementCount)
 			SHADER_PARAMETER(uint32, VerticalElementCount)
 			SHADER_PARAMETER_RDG_TEXTURE_UAV(RWTexture2D, OutTexture)
-			END_SHADER_PARAMETER_STRUCT()
+		END_SHADER_PARAMETER_STRUCT()
 
 	public:
 
@@ -50,13 +50,13 @@ namespace UE::RivermaxShaders
 		}
 
 		/** Allocates and setup shader parameter in the incoming graph builder */
-		FYUV8Bit422ToRGBACS::FParameters* AllocateAndSetParameters(FRDGBuilder& GraphBuilder, FRDGBufferRef YUVBuffer, FRDGTextureRef OutputTexture, const FMatrix& ColorTransform, const FVector& YUVOffset, int32 BufferElementsPerRow, int32 BufferLineCount);
+		RIVERMAXRENDERING_API FYUV8Bit422ToRGBACS::FParameters* AllocateAndSetParameters(FRDGBuilder& GraphBuilder, FRDGBufferRef YUVBuffer, FRDGTextureRef OutputTexture, const FMatrix& ColorTransform, const FVector& YUVOffset, int32 BufferElementsPerRow, int32 BufferLineCount);
 	};
 
 	/**
 	 * Compute shader to convert RGB to YUV422 8 bits
 	 */
-	class RIVERMAXRENDERING_API FRGBToYUV8Bit422CS : public FGlobalShader
+	class FRGBToYUV8Bit422CS : public FGlobalShader
 	{
 	public:
 		//Structure definition to match StructuredBuffer in RivermaxShaders.usf
@@ -65,7 +65,7 @@ namespace UE::RivermaxShaders
 			uint32 DWord0;
 		};
 
-		DECLARE_GLOBAL_SHADER(FRGBToYUV8Bit422CS);
+		DECLARE_EXPORTED_GLOBAL_SHADER(FRGBToYUV8Bit422CS, RIVERMAXRENDERING_API);
 
 		SHADER_USE_PARAMETER_STRUCT(FRGBToYUV8Bit422CS, FGlobalShader);
 
@@ -78,7 +78,7 @@ namespace UE::RivermaxShaders
 			SHADER_PARAMETER(uint32, VerticalElementCount)
 			SHADER_PARAMETER(uint32, InputTexturePixelsPerThread)
 			SHADER_PARAMETER_RDG_BUFFER_UAV(RWStructuredBuffer<FYUV8Bit422Buffer>, OutYUV4228bitBuffer)
-			END_SHADER_PARAMETER_STRUCT()
+		END_SHADER_PARAMETER_STRUCT()
 
 	public:
 
@@ -89,13 +89,13 @@ namespace UE::RivermaxShaders
 		}
 
 		/** Allocates and setup shader parameter in the incoming graph builder */
-		FRGBToYUV8Bit422CS::FParameters* AllocateAndSetParameters(FRDGBuilder& GraphBuilder, FRDGTextureRef RGBATexture, const FIntPoint& SourceSize, const FIntRect& SourceViewRect, const FIntPoint& OutputSize, const FMatrix& ColorTransform, const FVector& YUVOffset, bool bDoLinearToSrgb, FRDGBufferRef OutputBuffer);
+		RIVERMAXRENDERING_API FRGBToYUV8Bit422CS::FParameters* AllocateAndSetParameters(FRDGBuilder& GraphBuilder, FRDGTextureRef RGBATexture, const FIntPoint& SourceSize, const FIntRect& SourceViewRect, const FIntPoint& OutputSize, const FMatrix& ColorTransform, const FVector& YUVOffset, bool bDoLinearToSrgb, FRDGBufferRef OutputBuffer);
 	};
 
 	/**
 	 * Compute shader to convert packed 8 bits RGB to RGBA 8bits
 	 */
-	class RIVERMAXRENDERING_API FYUV10Bit422ToRGBACS : public FGlobalShader
+	class FYUV10Bit422ToRGBACS : public FGlobalShader
 	{
 	public:
 
@@ -109,7 +109,7 @@ namespace UE::RivermaxShaders
 			uint32 DWord4;
 		};
 
-		DECLARE_GLOBAL_SHADER(FYUV10Bit422ToRGBACS);
+		DECLARE_EXPORTED_GLOBAL_SHADER(FYUV10Bit422ToRGBACS, RIVERMAXRENDERING_API);
 		SHADER_USE_PARAMETER_STRUCT(FYUV10Bit422ToRGBACS, FGlobalShader);
 
 		class FSRGBToLinear : SHADER_PERMUTATION_BOOL("DO_SRGB_TO_LINEAR");
@@ -132,7 +132,7 @@ namespace UE::RivermaxShaders
 		}
 
 		/** Allocates and setup shader parameter in the incoming graph builder */
-		FYUV10Bit422ToRGBACS::FParameters* AllocateAndSetParameters(FRDGBuilder& GraphBuilder, FRDGBufferRef YUVBuffer, FRDGTextureRef OutputTexture, const FMatrix& ColorTransform, const FVector& YUVOffset, int32 BufferElementsPerRow, int32 BufferLineCount);
+		RIVERMAXRENDERING_API FYUV10Bit422ToRGBACS::FParameters* AllocateAndSetParameters(FRDGBuilder& GraphBuilder, FRDGBufferRef YUVBuffer, FRDGTextureRef OutputTexture, const FMatrix& ColorTransform, const FVector& YUVOffset, int32 BufferElementsPerRow, int32 BufferLineCount);
 	};
 
 	/**
@@ -140,7 +140,7 @@ namespace UE::RivermaxShaders
 	 *
 	 * This shader expects a single texture in PF_A2B10G10R10 format.
 	 */
-	class RIVERMAXRENDERING_API FRGBToYUV10Bit422LittleEndianCS : public FGlobalShader
+	class FRGBToYUV10Bit422LittleEndianCS : public FGlobalShader
 	{
 	public:
 		//Structure definition to match StructuredBuffer in RivermaxShaders.usf
@@ -153,7 +153,7 @@ namespace UE::RivermaxShaders
 			uint32 DWord4;
 		};
 
-		DECLARE_GLOBAL_SHADER(FRGBToYUV10Bit422LittleEndianCS);
+		DECLARE_EXPORTED_GLOBAL_SHADER(FRGBToYUV10Bit422LittleEndianCS, RIVERMAXRENDERING_API);
 
 		SHADER_USE_PARAMETER_STRUCT(FRGBToYUV10Bit422LittleEndianCS, FGlobalShader);
 
@@ -177,13 +177,13 @@ namespace UE::RivermaxShaders
 		}
 
 		/** Allocates and setup shader parameter in the incoming graph builder */
-		FRGBToYUV10Bit422LittleEndianCS::FParameters* AllocateAndSetParameters(FRDGBuilder& GraphBuilder, FRDGTextureRef RGBATexture, const FIntPoint& SourceSize, const FIntRect& SourceViewRect, const FIntPoint& OutputSize, const FMatrix& ColorTransform, const FVector& YUVOffset, bool bDoLinearToSrgb, FRDGBufferRef OutputBuffer);
+		RIVERMAXRENDERING_API FRGBToYUV10Bit422LittleEndianCS::FParameters* AllocateAndSetParameters(FRDGBuilder& GraphBuilder, FRDGTextureRef RGBATexture, const FIntPoint& SourceSize, const FIntRect& SourceViewRect, const FIntPoint& OutputSize, const FMatrix& ColorTransform, const FVector& YUVOffset, bool bDoLinearToSrgb, FRDGBufferRef OutputBuffer);
 	};
 
 	/**
 	 * Compute shader to convert packed 8 bits RGB to RGBA 8bits
 	 */
-	class RIVERMAXRENDERING_API FRGB8BitToRGBA8CS : public FGlobalShader
+	class FRGB8BitToRGBA8CS : public FGlobalShader
 	{
 	public:
 
@@ -196,7 +196,7 @@ namespace UE::RivermaxShaders
 			uint32 DWord2;
 		};
 
-		DECLARE_GLOBAL_SHADER(FRGB8BitToRGBA8CS);
+		DECLARE_EXPORTED_GLOBAL_SHADER(FRGB8BitToRGBA8CS, RIVERMAXRENDERING_API);
 		SHADER_USE_PARAMETER_STRUCT(FRGB8BitToRGBA8CS, FGlobalShader);
 
 		class FSRGBToLinear : SHADER_PERMUTATION_BOOL("DO_SRGB_TO_LINEAR");
@@ -218,13 +218,13 @@ namespace UE::RivermaxShaders
 		}
 
 		/** Allocates and setup shader parameter in the incoming graph builder */
-		FRGB8BitToRGBA8CS::FParameters* AllocateAndSetParameters(FRDGBuilder& GraphBuilder, FRDGBufferRef RGBBuffer, FRDGTextureRef OutputTexture, int32 BufferElementsPerRow, int32 BufferLineCount);
+		RIVERMAXRENDERING_API FRGB8BitToRGBA8CS::FParameters* AllocateAndSetParameters(FRDGBuilder& GraphBuilder, FRDGBufferRef RGBBuffer, FRDGTextureRef OutputTexture, int32 BufferElementsPerRow, int32 BufferLineCount);
 	};
 
 	/**
 	 * Compute shader to convert RGBA to packed 8 bits RGB
 	 */
-	class RIVERMAXRENDERING_API FRGBToRGB8BitCS : public FGlobalShader
+	class FRGBToRGB8BitCS : public FGlobalShader
 	{
 	public:
 		
@@ -237,7 +237,7 @@ namespace UE::RivermaxShaders
 			uint32 DWord2;
 		};
 
-		DECLARE_GLOBAL_SHADER(FRGBToRGB8BitCS);
+		DECLARE_EXPORTED_GLOBAL_SHADER(FRGBToRGB8BitCS, RIVERMAXRENDERING_API);
 
 		SHADER_USE_PARAMETER_STRUCT(FRGBToRGB8BitCS, FGlobalShader);
 
@@ -263,13 +263,13 @@ namespace UE::RivermaxShaders
 		}
 
 		/** Allocates and setup shader parameter in the incoming graph builder */
-		FRGBToRGB8BitCS::FParameters* AllocateAndSetParameters(FRDGBuilder& GraphBuilder, FRDGTextureRef RGBATexture, const FIntPoint& SourceSize, const FIntRect& SourceViewRect, const FIntPoint& OutputSize, FRDGBufferRef OutputBuffer);
+		RIVERMAXRENDERING_API FRGBToRGB8BitCS::FParameters* AllocateAndSetParameters(FRDGBuilder& GraphBuilder, FRDGTextureRef RGBATexture, const FIntPoint& SourceSize, const FIntRect& SourceViewRect, const FIntPoint& OutputSize, FRDGBufferRef OutputBuffer);
 	};
 
 	/**
 	 * Compute shader to convert RGBA to packed 10 bits RGB
 	 */
-	class RIVERMAXRENDERING_API FRGBToRGB10BitCS : public FGlobalShader
+	class FRGBToRGB10BitCS : public FGlobalShader
 	{
 	public:
 
@@ -281,7 +281,7 @@ namespace UE::RivermaxShaders
 			uint32 DWords[15];
 		};
 
-		DECLARE_GLOBAL_SHADER(FRGBToRGB10BitCS);
+		DECLARE_EXPORTED_GLOBAL_SHADER(FRGBToRGB10BitCS, RIVERMAXRENDERING_API);
 
 		SHADER_USE_PARAMETER_STRUCT(FRGBToRGB10BitCS, FGlobalShader);
 
@@ -296,7 +296,7 @@ namespace UE::RivermaxShaders
 			SHADER_PARAMETER(uint32, VerticalElementCount)
 			SHADER_PARAMETER(uint32, InputTexturePixelsPerThread)
 			SHADER_PARAMETER_RDG_BUFFER_UAV(RWStructuredBuffer<FRGB10BitBuffer>, OutRGB10Buffer)
-			END_SHADER_PARAMETER_STRUCT()
+		END_SHADER_PARAMETER_STRUCT()
 
 	public:
 
@@ -307,13 +307,13 @@ namespace UE::RivermaxShaders
 		}
 
 		/** Allocates and setup shader parameter in the incoming graph builder */
-		FRGBToRGB10BitCS::FParameters* AllocateAndSetParameters(FRDGBuilder& GraphBuilder, FRDGTextureRef RGBATexture, const FIntPoint& SourceSize, const FIntRect& SourceViewRect, const FIntPoint& OutputSize, FRDGBufferRef OutputBuffer);
+		RIVERMAXRENDERING_API FRGBToRGB10BitCS::FParameters* AllocateAndSetParameters(FRDGBuilder& GraphBuilder, FRDGTextureRef RGBATexture, const FIntPoint& SourceSize, const FIntRect& SourceViewRect, const FIntPoint& OutputSize, FRDGBufferRef OutputBuffer);
 	};
 
 	/**
 	 * Compute shader to convert packed 10 bits RGB to RGBA
 	 */
-	class RIVERMAXRENDERING_API FRGB10BitToRGBA10CS : public FGlobalShader
+	class FRGB10BitToRGBA10CS : public FGlobalShader
 	{
 	public:
 
@@ -325,7 +325,7 @@ namespace UE::RivermaxShaders
 			uint32 DWords[15];
 		};
 
-		DECLARE_GLOBAL_SHADER(FRGB10BitToRGBA10CS);
+		DECLARE_EXPORTED_GLOBAL_SHADER(FRGB10BitToRGBA10CS, RIVERMAXRENDERING_API);
 		SHADER_USE_PARAMETER_STRUCT(FRGB10BitToRGBA10CS, FGlobalShader);
 
 		class FSRGBToLinear : SHADER_PERMUTATION_BOOL("DO_SRGB_TO_LINEAR");
@@ -336,7 +336,7 @@ namespace UE::RivermaxShaders
 			SHADER_PARAMETER(uint32, HorizontalElementCount)
 			SHADER_PARAMETER(uint32, VerticalElementCount)
 			SHADER_PARAMETER_RDG_TEXTURE_UAV(RWTexture2D, OutTexture)
-			END_SHADER_PARAMETER_STRUCT()
+		END_SHADER_PARAMETER_STRUCT()
 
 	public:
 
@@ -347,13 +347,13 @@ namespace UE::RivermaxShaders
 		}
 
 		/** Allocates and setup shader parameter in the incoming graph builder */
-		FRGB10BitToRGBA10CS::FParameters* AllocateAndSetParameters(FRDGBuilder& GraphBuilder, FRDGBufferRef RGBBuffer, FRDGTextureRef OutputTexture, int32 BufferElementsPerRow, int32 BufferLineCount);
+		RIVERMAXRENDERING_API FRGB10BitToRGBA10CS::FParameters* AllocateAndSetParameters(FRDGBuilder& GraphBuilder, FRDGBufferRef RGBBuffer, FRDGTextureRef OutputTexture, int32 BufferElementsPerRow, int32 BufferLineCount);
 	};
 
 	/**
 	 * Compute shader to convert RGBA to packed 12 bits RGB
 	 */
-	class RIVERMAXRENDERING_API FRGBToRGB12BitCS : public FGlobalShader
+	class FRGBToRGB12BitCS : public FGlobalShader
 	{
 	public:
 
@@ -365,7 +365,7 @@ namespace UE::RivermaxShaders
 			uint32 DWords[9];
 		};
 
-		DECLARE_GLOBAL_SHADER(FRGBToRGB12BitCS);
+		DECLARE_EXPORTED_GLOBAL_SHADER(FRGBToRGB12BitCS, RIVERMAXRENDERING_API);
 
 		SHADER_USE_PARAMETER_STRUCT(FRGBToRGB12BitCS, FGlobalShader);
 
@@ -380,7 +380,7 @@ namespace UE::RivermaxShaders
 			SHADER_PARAMETER(uint32, VerticalElementCount)
 			SHADER_PARAMETER(uint32, InputTexturePixelsPerThread)
 			SHADER_PARAMETER_RDG_BUFFER_UAV(RWStructuredBuffer<FRGB12BitBuffer>, OutRGB12Buffer)
-			END_SHADER_PARAMETER_STRUCT()
+		END_SHADER_PARAMETER_STRUCT()
 
 	public:
 
@@ -391,13 +391,13 @@ namespace UE::RivermaxShaders
 		}
 
 		/** Allocates and setup shader parameter in the incoming graph builder */
-		FRGBToRGB12BitCS::FParameters* AllocateAndSetParameters(FRDGBuilder& GraphBuilder, FRDGTextureRef RGBATexture, const FIntPoint& SourceSize, const FIntRect& SourceViewRect, const FIntPoint& OutputSize, FRDGBufferRef OutputBuffer);
+		RIVERMAXRENDERING_API FRGBToRGB12BitCS::FParameters* AllocateAndSetParameters(FRDGBuilder& GraphBuilder, FRDGTextureRef RGBATexture, const FIntPoint& SourceSize, const FIntRect& SourceViewRect, const FIntPoint& OutputSize, FRDGBufferRef OutputBuffer);
 	};
 
 	/**
 	 * Compute shader to convert packed 12 bits RGB to RGBA
 	 */
-	class RIVERMAXRENDERING_API FRGB12BitToRGBA12CS : public FGlobalShader
+	class FRGB12BitToRGBA12CS : public FGlobalShader
 	{
 	public:
 
@@ -409,7 +409,7 @@ namespace UE::RivermaxShaders
 			uint32 DWords[9];
 		};
 
-		DECLARE_GLOBAL_SHADER(FRGB12BitToRGBA12CS);
+		DECLARE_EXPORTED_GLOBAL_SHADER(FRGB12BitToRGBA12CS, RIVERMAXRENDERING_API);
 		SHADER_USE_PARAMETER_STRUCT(FRGB12BitToRGBA12CS, FGlobalShader);
 
 		class FSRGBToLinear : SHADER_PERMUTATION_BOOL("DO_SRGB_TO_LINEAR");
@@ -420,7 +420,7 @@ namespace UE::RivermaxShaders
 			SHADER_PARAMETER(uint32, HorizontalElementCount)
 			SHADER_PARAMETER(uint32, VerticalElementCount)
 			SHADER_PARAMETER_RDG_TEXTURE_UAV(RWTexture2D, OutTexture)
-			END_SHADER_PARAMETER_STRUCT()
+		END_SHADER_PARAMETER_STRUCT()
 
 	public:
 
@@ -431,17 +431,17 @@ namespace UE::RivermaxShaders
 		}
 
 		/** Allocates and setup shader parameter in the incoming graph builder */
-		FRGB12BitToRGBA12CS::FParameters* AllocateAndSetParameters(FRDGBuilder& GraphBuilder, FRDGBufferRef RGBBuffer, FRDGTextureRef OutputTexture, int32 BufferElementsPerRow, int32 BufferLineCount);
+		RIVERMAXRENDERING_API FRGB12BitToRGBA12CS::FParameters* AllocateAndSetParameters(FRDGBuilder& GraphBuilder, FRDGBufferRef RGBBuffer, FRDGTextureRef OutputTexture, int32 BufferElementsPerRow, int32 BufferLineCount);
 	};
 
 	/**
 	 * Compute shader to convert RGBA to Float16 RGB
 	 */
-	class RIVERMAXRENDERING_API FRGBToRGB16fCS : public FGlobalShader
+	class FRGBToRGB16fCS : public FGlobalShader
 	{
 	public:
 
-		DECLARE_GLOBAL_SHADER(FRGBToRGB16fCS);
+		DECLARE_EXPORTED_GLOBAL_SHADER(FRGBToRGB16fCS, RIVERMAXRENDERING_API);
 
 		SHADER_USE_PARAMETER_STRUCT(FRGBToRGB16fCS, FGlobalShader);
 
@@ -474,13 +474,13 @@ namespace UE::RivermaxShaders
 		}
 
 		/** Allocates and setup shader parameter in the incoming graph builder */
-		FRGBToRGB16fCS::FParameters* AllocateAndSetParameters(FRDGBuilder& GraphBuilder, FRDGTextureRef RGBATexture, const FIntPoint& SourceSize, const FIntRect& SourceViewRect, const FIntPoint& OutputSize, FRDGBufferRef OutputBuffer);
+		RIVERMAXRENDERING_API FRGBToRGB16fCS::FParameters* AllocateAndSetParameters(FRDGBuilder& GraphBuilder, FRDGTextureRef RGBATexture, const FIntPoint& SourceSize, const FIntRect& SourceViewRect, const FIntPoint& OutputSize, FRDGBufferRef OutputBuffer);
 	};
 
 	/**
 	 * Compute shader to convert 16 bits float RGB to RGBA
 	 */
-	class RIVERMAXRENDERING_API FRGB16fBitToRGBA16fCS : public FGlobalShader
+	class FRGB16fBitToRGBA16fCS : public FGlobalShader
 	{
 	public:
 
@@ -491,7 +491,7 @@ namespace UE::RivermaxShaders
 			uint32 DWords[3];
 		};
 
-		DECLARE_GLOBAL_SHADER(FRGB16fBitToRGBA16fCS);
+		DECLARE_EXPORTED_GLOBAL_SHADER(FRGB16fBitToRGBA16fCS, RIVERMAXRENDERING_API);
 		SHADER_USE_PARAMETER_STRUCT(FRGB16fBitToRGBA16fCS, FGlobalShader);
 
 		class FSRGBToLinear : SHADER_PERMUTATION_BOOL("DO_SRGB_TO_LINEAR");
@@ -502,7 +502,7 @@ namespace UE::RivermaxShaders
 			SHADER_PARAMETER(uint32, HorizontalElementCount)
 			SHADER_PARAMETER(uint32, VerticalElementCount)
 			SHADER_PARAMETER_RDG_TEXTURE_UAV(RWTexture2D, OutTexture)
-			END_SHADER_PARAMETER_STRUCT()
+		END_SHADER_PARAMETER_STRUCT()
 
 	public:
 
@@ -513,7 +513,7 @@ namespace UE::RivermaxShaders
 		}
 
 		/** Allocates and setup shader parameter in the incoming graph builder */
-		FRGB16fBitToRGBA16fCS::FParameters* AllocateAndSetParameters(FRDGBuilder& GraphBuilder, FRDGBufferRef RGBBuffer, FRDGTextureRef OutputTexture, int32 BufferElementsPerRow, int32 BufferLineCount);
+		RIVERMAXRENDERING_API FRGB16fBitToRGBA16fCS::FParameters* AllocateAndSetParameters(FRDGBuilder& GraphBuilder, FRDGBufferRef RGBBuffer, FRDGTextureRef OutputTexture, int32 BufferElementsPerRow, int32 BufferLineCount);
 	};
 }
 
diff --git a/Engine/Source/Runtime/Engine/Public/SimpleElementShaders.h b/Engine/Source/Runtime/Engine/Public/SimpleElementShaders.h
--- a/Engine/Source/Runtime/Engine/Public/SimpleElementShaders.h
+++ b/Engine/Source/Runtime/Engine/Public/SimpleElementShaders.h
@@ -20,20 +20,20 @@ struct FRelativeViewMatrices;
 /**
  * A vertex shader for rendering a texture on a simple element.
  */
-class ENGINE_API FSimpleElementVS : public FGlobalShader
+class FSimpleElementVS : public FGlobalShader
 {
-	DECLARE_SHADER_TYPE(FSimpleElementVS,Global);
+	DECLARE_EXPORTED_GLOBAL_SHADER(FSimpleElementVS, ENGINE_API);
 public:
 	FSimpleElementVS(const ShaderMetaType::CompiledShaderInitializerType& Initializer);
 	FSimpleElementVS() {}
 
-	void SetParameters(FRHIBatchedShaderParameters& BatchedParameters, const FMatrix& WorldToClipMatrix);
-	void SetParameters(FRHIBatchedShaderParameters& BatchedParameters, const FRelativeViewMatrices& Matrices);
+	ENGINE_API void SetParameters(FRHIBatchedShaderParameters& BatchedParameters, const FMatrix& WorldToClipMatrix);
+	ENGINE_API void SetParameters(FRHIBatchedShaderParameters& BatchedParameters, const FRelativeViewMatrices& Matrices);
 
 	UE_DEPRECATED(5.3, "SetParameters with FRHIBatchedShaderParameters should be used.")
-	void SetParameters(FRHICommandList& RHICmdList, const FMatrix& WorldToClipMatrix);
+	ENGINE_API void SetParameters(FRHICommandList& RHICmdList, const FMatrix& WorldToClipMatrix);
 	UE_DEPRECATED(5.3, "SetParameters with FRHIBatchedShaderParameters should be used.")
-	void SetParameters(FRHICommandList& RHICmdList, const FRelativeViewMatrices& Matrices);
+	ENGINE_API void SetParameters(FRHICommandList& RHICmdList, const FRelativeViewMatrices& Matrices);
 
 	static void ModifyCompilationEnvironment(const FGlobalShaderPermutationParameters& Parameters, FShaderCompilerEnvironment& OutEnvironment);
 
diff --git a/Engine/Source/Runtime/RenderCore/Private/RenderGraphUtils.cpp b/Engine/Source/Runtime/RenderCore/Private/RenderGraphUtils.cpp
--- a/Engine/Source/Runtime/RenderCore/Private/RenderGraphUtils.cpp
+++ b/Engine/Source/Runtime/RenderCore/Private/RenderGraphUtils.cpp
@@ -228,7 +228,7 @@ BEGIN_SHADER_PARAMETER_STRUCT(FCopyTextureParameters, )
 	RDG_TEXTURE_ACCESS(Output, ERHIAccess::CopyDest)
 END_SHADER_PARAMETER_STRUCT()
 
-class RENDERCORE_API FDrawTexturePS : public FGlobalShader
+class FDrawTexturePS : public FGlobalShader
 {
 	DECLARE_GLOBAL_SHADER(FDrawTexturePS);
 	SHADER_USE_PARAMETER_STRUCT(FDrawTexturePS, FGlobalShader);
diff --git a/Engine/Source/Runtime/RenderCore/Public/CommonRenderResources.h b/Engine/Source/Runtime/RenderCore/Public/CommonRenderResources.h
--- a/Engine/Source/Runtime/RenderCore/Public/CommonRenderResources.h
+++ b/Engine/Source/Runtime/RenderCore/Public/CommonRenderResources.h
@@ -106,15 +106,11 @@ extern RENDERCORE_API TGlobalResource<FScreenRectangleIndexBuffer> GScreenRectan
 
 /** Vertex shader to draw a screen quad that works on all platforms. Does not have any shader parameters.
  * The pixel shader should just use SV_Position. */
-class RENDERCORE_API FScreenVertexShaderVS : public FGlobalShader
+class FScreenVertexShaderVS : public FGlobalShader
 {
-	DECLARE_GLOBAL_SHADER(FScreenVertexShaderVS);
+	DECLARE_EXPORTED_GLOBAL_SHADER(FScreenVertexShaderVS, RENDERCORE_API);
 	SHADER_USE_PARAMETER_STRUCT(FScreenVertexShaderVS, FGlobalShader);
 
-	static bool ShouldCompilePermutation(const FGlobalShaderPermutationParameters& Parameters) {
-		return true;
-	}
-
 	BEGIN_SHADER_PARAMETER_STRUCT(FParameters, )
 	END_SHADER_PARAMETER_STRUCT()
 };
@@ -139,16 +135,11 @@ class RENDERCORE_API FInstancedScreenVertexShaderVS : public FScreenVertexShader
 };
 
 /** Pixel shader to copy pixels from src to dst performing a format change that works on all platforms. */
-class RENDERCORE_API FCopyRectPS : public FGlobalShader
+class FCopyRectPS : public FGlobalShader
 {
-	DECLARE_GLOBAL_SHADER(FCopyRectPS);
+	DECLARE_EXPORTED_GLOBAL_SHADER(FCopyRectPS, RENDERCORE_API);
 	SHADER_USE_PARAMETER_STRUCT(FCopyRectPS, FGlobalShader);
 
-	static bool ShouldCompilePermutation(const FGlobalShaderPermutationParameters& Parameters) 
-	{
-		return true;
-	}
-
 	BEGIN_SHADER_PARAMETER_STRUCT(FParameters, )
 		SHADER_PARAMETER_RDG_TEXTURE(Texture2D, InputTexture)
 		SHADER_PARAMETER_SAMPLER(SamplerState, InputSampler)
diff --git a/Engine/Source/Runtime/RenderCore/Public/GlobalShader.h b/Engine/Source/Runtime/RenderCore/Public/GlobalShader.h
--- a/Engine/Source/Runtime/RenderCore/Public/GlobalShader.h
+++ b/Engine/Source/Runtime/RenderCore/Public/GlobalShader.h
@@ -78,7 +78,7 @@ struct FGlobalShaderPermutationParameters : public FShaderPermutationParameters
  * A shader meta type for the simplest shaders; shaders which are not material or vertex factory linked.
  * There should only a single instance of each simple shader type.
  */
-class RENDERCORE_API FGlobalShaderType : public FShaderType
+class FGlobalShaderType : public FShaderType
 {
 	friend class FGlobalShaderTypeCompiler;
 public:
@@ -138,13 +138,13 @@ public:
 	 * @param Platform - Platform to compile for.
 	 * @param Environment - The shader compile environment that the function modifies.
 	 */
-	void SetupCompileEnvironment(EShaderPlatform Platform, int32 PermutationId, EShaderPermutationFlags Flags, FShaderCompilerEnvironment& Environment) const;
+	RENDERCORE_API void SetupCompileEnvironment(EShaderPlatform Platform, int32 PermutationId, EShaderPermutationFlags Flags, FShaderCompilerEnvironment& Environment) const;
 #endif // WITH_EDITOR
 
-	static bool ShouldCompilePipeline(const FShaderPipelineType* ShaderPipelineType, EShaderPlatform Platform, EShaderPermutationFlags Flags);
+	RENDERCORE_API static bool ShouldCompilePipeline(const FShaderPipelineType* ShaderPipelineType, EShaderPlatform Platform, EShaderPermutationFlags Flags);
 };
 
-class RENDERCORE_API FGlobalShaderMapContent : public FShaderMapContent
+class FGlobalShaderMapContent : public FShaderMapContent
 {
 	using Super = FShaderMapContent;
 	friend class FGlobalShaderMap;
@@ -162,14 +162,14 @@ private:
 	LAYOUT_FIELD(FHashedName, HashedSourceFilename);
 };
 
-class RENDERCORE_API FGlobalShaderMapSection : public TShaderMap<FGlobalShaderMapContent, FShaderMapPointerTable>
+class FGlobalShaderMapSection : public TShaderMap<FGlobalShaderMapContent, FShaderMapPointerTable>
 {
 	using Super = TShaderMap<FGlobalShaderMapContent, FShaderMapPointerTable>;
 	friend class FGlobalShaderMap;
 public:
-	static FGlobalShaderMapSection* CreateFromArchive(FArchive& Ar);
+	RENDERCORE_API static FGlobalShaderMapSection* CreateFromArchive(FArchive& Ar);
 
-	bool Serialize(FArchive& Ar);
+	RENDERCORE_API bool Serialize(FArchive& Ar);
 private:
 	inline FGlobalShaderMapSection() 
 	{ 
@@ -188,14 +188,14 @@ private:
 	FShaderPipelineRef GetShaderPipeline(const FShaderPipelineType* PipelineType) const;
 };
 
-class RENDERCORE_API FGlobalShaderMap
+class FGlobalShaderMap
 {
 public:
-	explicit FGlobalShaderMap(EShaderPlatform InPlatform);
-	~FGlobalShaderMap();
+	RENDERCORE_API explicit FGlobalShaderMap(EShaderPlatform InPlatform);
+	RENDERCORE_API ~FGlobalShaderMap();
 
-	TShaderRef<FShader> GetShader(FShaderType* ShaderType, int32 PermutationId = 0) const;
-	FShaderPipelineRef GetShaderPipeline(const FShaderPipelineType* PipelineType) const;
+	RENDERCORE_API TShaderRef<FShader> GetShader(FShaderType* ShaderType, int32 PermutationId = 0) const;
+	RENDERCORE_API FShaderPipelineRef GetShaderPipeline(const FShaderPipelineType* PipelineType) const;
 
 	template<typename ShaderType>
 	TShaderRef<ShaderType> GetShader() const
@@ -229,32 +229,32 @@ public:
 		return GetShaderPipeline(ShaderPipelineType).IsValid();
 	}
 
-	bool IsEmpty() const;
+	RENDERCORE_API bool IsEmpty() const;
 
 	/** Whether the shadermap has all the shader types it needs (i.e. ones that returned ShouldCompilePermutation) */
-	bool IsComplete(const ITargetPlatform* TargetPlatform) const;
+	RENDERCORE_API bool IsComplete(const ITargetPlatform* TargetPlatform) const;
 
-	void Empty();
-	void ReleaseAllSections();
+	RENDERCORE_API void Empty();
+	RENDERCORE_API void ReleaseAllSections();
 
-	FShader* FindOrAddShader(const FShaderType* ShaderType, int32 PermutationId, FShader* Shader);
-	FShaderPipeline* FindOrAddShaderPipeline(const FShaderPipelineType* ShaderPipelineType, FShaderPipeline* ShaderPipeline);
+	RENDERCORE_API FShader* FindOrAddShader(const FShaderType* ShaderType, int32 PermutationId, FShader* Shader);
+	RENDERCORE_API FShaderPipeline* FindOrAddShaderPipeline(const FShaderPipelineType* ShaderPipelineType, FShaderPipeline* ShaderPipeline);
 
-	void RemoveShaderTypePermutaion(const FShaderType* Type, int32 PermutationId);
-	void RemoveShaderPipelineType(const FShaderPipelineType* ShaderPipelineType);
+	RENDERCORE_API void RemoveShaderTypePermutaion(const FShaderType* Type, int32 PermutationId);
+	RENDERCORE_API void RemoveShaderPipelineType(const FShaderPipelineType* ShaderPipelineType);
 
-	void AddSection(FGlobalShaderMapSection* InSection);
-	FGlobalShaderMapSection* FindSection(const FHashedName& HashedShaderFilename);
-	FGlobalShaderMapSection* FindOrAddSection(const FShaderType* ShaderType);
+	RENDERCORE_API void AddSection(FGlobalShaderMapSection* InSection);
+	RENDERCORE_API FGlobalShaderMapSection* FindSection(const FHashedName& HashedShaderFilename);
+	RENDERCORE_API FGlobalShaderMapSection* FindOrAddSection(const FShaderType* ShaderType);
 	
-	void LoadFromGlobalArchive(FArchive& Ar);
-	void SaveToGlobalArchive(FArchive& Ar);
+	RENDERCORE_API void LoadFromGlobalArchive(FArchive& Ar);
+	RENDERCORE_API void SaveToGlobalArchive(FArchive& Ar);
 
-	void BeginCreateAllShaders();
+	RENDERCORE_API void BeginCreateAllShaders();
 
 #if WITH_EDITOR
-	void GetOutdatedTypes(TArray<const FShaderType*>& OutdatedShaderTypes, TArray<const FShaderPipelineType*>& OutdatedShaderPipelineTypes, TArray<const FVertexFactoryType*>& OutdatedFactoryTypes) const;
-	void SaveShaderStableKeys(EShaderPlatform TargetShaderPlatform);
+	RENDERCORE_API void GetOutdatedTypes(TArray<const FShaderType*>& OutdatedShaderTypes, TArray<const FShaderPipelineType*>& OutdatedShaderPipelineTypes, TArray<const FVertexFactoryType*>& OutdatedFactoryTypes) const;
+	RENDERCORE_API void SaveShaderStableKeys(EShaderPlatform TargetShaderPlatform);
 #endif // WITH_EDITOR
 
 	const FGlobalShaderMapSection* GetFirstSection()
@@ -387,13 +387,13 @@ inline FGlobalShaderMap* GetGlobalShaderMap(ERHIFeatureLevel::Type FeatureLevel)
  * done in a .cpp file regardless of whether FMyGlobalShaderPS is in a header or not.
  * IMPLEMENT_GLOBAL_SHADER(FMyGlobalShaderPS, "/Engine/Private/MyShaderFile.usf", "MainPS", SF_Pixel);
  *
- * When the shader class is a public header, let say in RenderCore module public header, the shader class then should have the RENDERCORE_API
+ * When the shader class is a public header, let say in RenderCore module public header, the class declaration macro then should have the RENDERCORE_API
  * like this:
  *
- * class RENDERCORE_API FMyGlobalShaderPS : public FGlobalShader
+ * class FMyGlobalShaderPS : public FGlobalShader
  * {
  *		// Setup the shader's boiler plate.
- *		DECLARE_GLOBAL_SHADER(FMyGlobalShaderPS);
+ *		DECLARE_EXPORTED_GLOBAL_SHADER(FMyGlobalShaderPS, RENDERCORE_API);
  *
  *		// ...
  * };
diff --git a/Engine/Source/Runtime/RenderCore/Public/MediaShaders.h b/Engine/Source/Runtime/RenderCore/Public/MediaShaders.h
--- a/Engine/Source/Runtime/RenderCore/Public/MediaShaders.h
+++ b/Engine/Source/Runtime/RenderCore/Public/MediaShaders.h
@@ -578,11 +578,10 @@ END_SHADER_PARAMETER_STRUCT()
 /**
  * Pixel shader to convert RGB 8 bits to UYVY 8 bits
  */
-	class RENDERCORE_API FRGB8toUYVY8ConvertPS : public FGlobalShader
+class FRGB8toUYVY8ConvertPS : public FGlobalShader
 {
 public:
-
-	DECLARE_GLOBAL_SHADER(FRGB8toUYVY8ConvertPS);
+	DECLARE_EXPORTED_GLOBAL_SHADER(FRGB8toUYVY8ConvertPS, RENDERCORE_API);
 
 	SHADER_USE_PARAMETER_STRUCT(FRGB8toUYVY8ConvertPS, FGlobalShader);
 
@@ -592,24 +591,18 @@ public:
 	END_SHADER_PARAMETER_STRUCT()
 
 public:
-	UE_DEPRECATED(5.1, "SetParameters has been deprecated while moving to a RDG based pipeline. Please use AllocateAndSetParameters instead.")
-	void SetParameters(FRHICommandList& RHICmdList, TRefCountPtr<FRHITexture2D> RGBATexture, const FMatrix& ColorTransform, const FVector& YUVOffset, bool LinearToSrgb)
-	{
-	}
-
 	/** Allocates and setup shader parameter in the incoming graph builder */
-	FRGB8toUYVY8ConvertPS::FParameters* AllocateAndSetParameters(FRDGBuilder& GraphBuilder, FRDGTextureRef RGBATexture, const FMatrix& ColorTransform, const FVector& YUVOffset, bool bDoLinearToSrgb, FRDGTextureRef OutputTexture);
+	RENDERCORE_API FRGB8toUYVY8ConvertPS::FParameters* AllocateAndSetParameters(FRDGBuilder& GraphBuilder, FRDGTextureRef RGBATexture, const FMatrix& ColorTransform, const FVector& YUVOffset, bool bDoLinearToSrgb, FRDGTextureRef OutputTexture);
 };
 
 
 /**
  * Pixel shader to convert RGB 10 bits to YUV v210
  */
-class RENDERCORE_API FRGB10toYUVv210ConvertPS : public FGlobalShader
+class FRGB10toYUVv210ConvertPS : public FGlobalShader
 {
 public:
-
-	DECLARE_GLOBAL_SHADER(FRGB10toYUVv210ConvertPS);
+	DECLARE_EXPORTED_GLOBAL_SHADER(FRGB10toYUVv210ConvertPS, RENDERCORE_API);
 
 	SHADER_USE_PARAMETER_STRUCT(FRGB10toYUVv210ConvertPS, FGlobalShader);
 
@@ -620,13 +613,8 @@ public:
 	END_SHADER_PARAMETER_STRUCT()
 
 public:
-	UE_DEPRECATED(5.1, "SetParameters has been deprecated while moving to a RDG based pipeline. Please use AllocateAndSetParameters instead.")
-	void SetParameters(FRHICommandList& RHICmdList, TRefCountPtr<FRHITexture2D> RGBATexture, const FMatrix& ColorTransform, const FVector& YUVOffset, bool LinearToSrgb)
-	{
-	}
-
 	/** Allocates and setup shader parameter in the incoming graph builder */
-	FRGB10toYUVv210ConvertPS::FParameters* AllocateAndSetParameters(FRDGBuilder& GraphBuilder, FRDGTextureRef RGBATexture, const FMatrix& ColorTransform, const FVector& YUVOffset, bool bDoLinearToSrgb, FRDGTextureRef OutputTexture);
+	RENDERCORE_API FRGB10toYUVv210ConvertPS::FParameters* AllocateAndSetParameters(FRDGBuilder& GraphBuilder, FRDGTextureRef RGBATexture, const FMatrix& ColorTransform, const FVector& YUVOffset, bool bDoLinearToSrgb, FRDGTextureRef OutputTexture);
 };
 
 /**
@@ -634,11 +622,10 @@ public:
  *
  * General conversion shader that is only used to swizzle shaders that do not require any color conversion. RGB to BGR, RGB10A2 to RGBA8 etc
  */
-class RENDERCORE_API FModifyAlphaSwizzleRgbaPS	: public FGlobalShader
+class FModifyAlphaSwizzleRgbaPS	: public FGlobalShader
 {
 public:
-
-	DECLARE_GLOBAL_SHADER(FModifyAlphaSwizzleRgbaPS);
+	DECLARE_EXPORTED_GLOBAL_SHADER(FModifyAlphaSwizzleRgbaPS, RENDERCORE_API);
 
 	SHADER_USE_PARAMETER_STRUCT(FModifyAlphaSwizzleRgbaPS, FGlobalShader);
 
@@ -652,13 +639,8 @@ public:
 	END_SHADER_PARAMETER_STRUCT()
 
 public:
-	UE_DEPRECATED(5.1, "SetParameters has been deprecated while moving to a RDG based pipeline. Please use AllocateAndSetParameters instead.")
-	void SetParameters(FRHICommandList& RHICmdList, TRefCountPtr<FRHITexture2D> RGBATexture)
-	{
-	}
-
 	/** Allocates and setup shader parameter in the incoming graph builder */
-	FModifyAlphaSwizzleRgbaPS::FParameters* AllocateAndSetParameters(FRDGBuilder& GraphBuilder, FRDGTextureRef RGBATexture, FRDGTextureRef OutputTexture);
+	RENDERCORE_API FModifyAlphaSwizzleRgbaPS::FParameters* AllocateAndSetParameters(FRDGBuilder& GraphBuilder, FRDGTextureRef RGBATexture, FRDGTextureRef OutputTexture);
 };
 
 
diff --git a/Engine/Source/Runtime/RenderCore/Public/OneColorShader.h b/Engine/Source/Runtime/RenderCore/Public/OneColorShader.h
--- a/Engine/Source/Runtime/RenderCore/Public/OneColorShader.h
+++ b/Engine/Source/Runtime/RenderCore/Public/OneColorShader.h
@@ -75,27 +75,27 @@ private:
 /**
  * Pixel shader for rendering a single, constant color.
  */
-class RENDERCORE_API FOneColorPS : public FGlobalShader
+class FOneColorPS : public FGlobalShader
 {
-	DECLARE_GLOBAL_SHADER(FOneColorPS);
+	DECLARE_EXPORTED_GLOBAL_SHADER(FOneColorPS, RENDERCORE_API);
 public:
 	SHADER_USE_PARAMETER_STRUCT(FOneColorPS, FGlobalShader);
 
 	BEGIN_SHADER_PARAMETER_STRUCT(FParameters, RENDERCORE_API)
 		SHADER_PARAMETER_ARRAY(FLinearColor, DrawColorMRT, [MaxSimultaneousRenderTargets])
 		RENDER_TARGET_BINDING_SLOTS()
-		END_SHADER_PARAMETER_STRUCT()
+	END_SHADER_PARAMETER_STRUCT()
 
 	void FillParameters(FParameters& Parameters, const FLinearColor* Colors, int32 NumColors);
-	static void SetColors(FRHICommandList& RHICmdList, const TShaderMapRef<FOneColorPS>& Shader, const FLinearColor* Colors, int32 NumColors);
+	RENDERCORE_API static void SetColors(FRHICommandList& RHICmdList, const TShaderMapRef<FOneColorPS>& Shader, const FLinearColor* Colors, int32 NumColors);
 };
 
 /**
  * Pixel shader for rendering a single, constant color to MRTs.
  */
-class RENDERCORE_API TOneColorPixelShaderMRT : public FOneColorPS
+class TOneColorPixelShaderMRT : public FOneColorPS
 {
-	DECLARE_GLOBAL_SHADER(TOneColorPixelShaderMRT);
+	DECLARE_EXPORTED_GLOBAL_SHADER(TOneColorPixelShaderMRT, RENDERCORE_API);
 public:
 	class TOneColorPixelShader128bitRT : SHADER_PERMUTATION_BOOL("b128BITRENDERTARGET");
 	class TOneColorPixelShaderNumOutputs : SHADER_PERMUTATION_RANGE_INT("NUM_OUTPUTS", 1, 8);
@@ -106,15 +106,15 @@ public:
 
 	static bool ShouldCompilePermutation(const FGlobalShaderPermutationParameters& Parameters);
 	static void ModifyCompilationEnvironment(const FGlobalShaderPermutationParameters& Parameters, FShaderCompilerEnvironment& OutEnvironment);
-	static void SetColors(FRHICommandList& RHICmdList, const TShaderMapRef<TOneColorPixelShaderMRT>& Shader, const FLinearColor* Colors, int32 NumColors);
+	RENDERCORE_API static void SetColors(FRHICommandList& RHICmdList, const TShaderMapRef<TOneColorPixelShaderMRT>& Shader, const FLinearColor* Colors, int32 NumColors);
 };
 
 /**
  * Compute shader for writing values
  */
-class RENDERCORE_API FFillTextureCS : public FGlobalShader
+class FFillTextureCS : public FGlobalShader
 {
-	DECLARE_GLOBAL_SHADER(FFillTextureCS);
+	DECLARE_EXPORTED_GLOBAL_SHADER(FFillTextureCS, RENDERCORE_API);
 public:
 	FFillTextureCS();
 	FFillTextureCS(const ShaderMetaType::CompiledShaderInitializerType& Initializer);
@@ -128,9 +128,9 @@ public:
 	LAYOUT_FIELD(FShaderResourceParameter, FillTexture);
 };
 
-class RENDERCORE_API FLongGPUTaskPS : public FGlobalShader
+class FLongGPUTaskPS : public FGlobalShader
 {
-	DECLARE_GLOBAL_SHADER(FLongGPUTaskPS);
+	DECLARE_EXPORTED_GLOBAL_SHADER(FLongGPUTaskPS, RENDERCORE_API);
 public:
 	FLongGPUTaskPS();
 	FLongGPUTaskPS(const ShaderMetaType::CompiledShaderInitializerType& Initializer);
diff --git a/Engine/Source/Runtime/RenderCore/Public/ResolveShader.h b/Engine/Source/Runtime/RenderCore/Public/ResolveShader.h
--- a/Engine/Source/Runtime/RenderCore/Public/ResolveShader.h
+++ b/Engine/Source/Runtime/RenderCore/Public/ResolveShader.h
@@ -17,9 +17,9 @@ struct FResolveRect;
 
 struct FDummyResolveParameter {};
 
-class RENDERCORE_API FResolveDepthPS : public FGlobalShader
+class FResolveDepthPS : public FGlobalShader
 {
-	DECLARE_GLOBAL_SHADER(FResolveDepthPS);
+	DECLARE_EXPORTED_GLOBAL_SHADER(FResolveDepthPS, RENDERCORE_API);
 public:
 	typedef FDummyResolveParameter FParameter;
 		
@@ -32,10 +32,11 @@ public:
 	LAYOUT_FIELD(FShaderResourceParameter, UnresolvedSurface);
 };
 
-class RENDERCORE_API FResolveDepth2XPS : public FResolveDepthPS
+class FResolveDepth2XPS : public FResolveDepthPS
 {
-	DECLARE_GLOBAL_SHADER(FResolveDepth2XPS);
+	DECLARE_EXPORTED_GLOBAL_SHADER(FResolveDepth2XPS, RENDERCORE_API);
 public:
+
 	typedef FDummyResolveParameter FParameter;
 
 	FResolveDepth2XPS();
@@ -44,9 +45,9 @@ public:
 	static void ModifyCompilationEnvironment(const FGlobalShaderPermutationParameters& Parameters, FShaderCompilerEnvironment& OutEnvironment);
 };
 
-class RENDERCORE_API FResolveDepth4XPS : public FResolveDepthPS
+class FResolveDepth4XPS : public FResolveDepthPS
 {
-	DECLARE_GLOBAL_SHADER(FResolveDepth4XPS);
+	DECLARE_EXPORTED_GLOBAL_SHADER(FResolveDepth4XPS, RENDERCORE_API);
 public:
 	typedef FDummyResolveParameter FParameter;
 
@@ -56,9 +57,9 @@ public:
 	static void ModifyCompilationEnvironment(const FGlobalShaderPermutationParameters& Parameters, FShaderCompilerEnvironment& OutEnvironment);
 };
 
-class RENDERCORE_API FResolveDepth8XPS : public FResolveDepthPS
+class FResolveDepth8XPS : public FResolveDepthPS
 {
-	DECLARE_GLOBAL_SHADER(FResolveDepth8XPS);
+	DECLARE_EXPORTED_GLOBAL_SHADER(FResolveDepth8XPS, RENDERCORE_API);
 public:
 	typedef FDummyResolveParameter FParameter;
 
@@ -69,9 +70,9 @@ public:
 	static void ModifyCompilationEnvironment(const FGlobalShaderPermutationParameters& Parameters, FShaderCompilerEnvironment& OutEnvironment);
 };
 
-class RENDERCORE_API FResolveDepthArrayPS : public FResolveDepthPS
+class FResolveDepthArrayPS : public FResolveDepthPS
 {
-	DECLARE_GLOBAL_SHADER(FResolveDepthArrayPS);
+	DECLARE_EXPORTED_GLOBAL_SHADER(FResolveDepthArrayPS, RENDERCORE_API);
 public:
 	typedef FDummyResolveParameter FParameter;
 
@@ -82,9 +83,9 @@ public:
 	static void ModifyCompilationEnvironment(const FGlobalShaderPermutationParameters& Parameters, FShaderCompilerEnvironment& OutEnvironment);
 };
 
-class RENDERCORE_API FResolveDepthArray2XPS : public FResolveDepthArrayPS
+class FResolveDepthArray2XPS : public FResolveDepthArrayPS
 {
-	DECLARE_GLOBAL_SHADER(FResolveDepthArray2XPS);
+	DECLARE_EXPORTED_GLOBAL_SHADER(FResolveDepthArray2XPS, RENDERCORE_API);
 public:
 	typedef FDummyResolveParameter FParameter;
 
@@ -95,9 +96,9 @@ public:
 	static void ModifyCompilationEnvironment(const FGlobalShaderPermutationParameters& Parameters, FShaderCompilerEnvironment& OutEnvironment);
 };
 
-class RENDERCORE_API FResolveDepthArray4XPS : public FResolveDepthArrayPS
+class FResolveDepthArray4XPS : public FResolveDepthArrayPS
 {
-	DECLARE_GLOBAL_SHADER(FResolveDepthArray4XPS);
+	DECLARE_EXPORTED_GLOBAL_SHADER(FResolveDepthArray4XPS, RENDERCORE_API);
 public:
 	typedef FDummyResolveParameter FParameter;
 
@@ -108,23 +109,23 @@ public:
 	static void ModifyCompilationEnvironment(const FGlobalShaderPermutationParameters& Parameters, FShaderCompilerEnvironment& OutEnvironment);
 };
 
-class RENDERCORE_API FResolveDepthArray8XPS : public FResolveDepthArrayPS
+class FResolveDepthArray8XPS : public FResolveDepthArrayPS
 {
-	DECLARE_GLOBAL_SHADER(FResolveDepthArray8XPS);
+	DECLARE_EXPORTED_GLOBAL_SHADER(FResolveDepthArray8XPS, RENDERCORE_API);
 public:
 	typedef FDummyResolveParameter FParameter;
 
 	FResolveDepthArray8XPS();
 	FResolveDepthArray8XPS(const ShaderMetaType::CompiledShaderInitializerType& Initializer);
 
-	static bool ShouldCache(EShaderPlatform Platform);
+	RENDERCORE_API static bool ShouldCache(EShaderPlatform Platform);
 	static bool ShouldCompilePermutation(const FGlobalShaderPermutationParameters& Parameters);
 	static void ModifyCompilationEnvironment(const FGlobalShaderPermutationParameters& Parameters, FShaderCompilerEnvironment& OutEnvironment);
 };
 
-class RENDERCORE_API FResolveSingleSamplePS : public FGlobalShader
+class FResolveSingleSamplePS : public FGlobalShader
 {
-	DECLARE_GLOBAL_SHADER(FResolveSingleSamplePS);
+	DECLARE_EXPORTED_GLOBAL_SHADER(FResolveSingleSamplePS, RENDERCORE_API);
 public:
 	typedef uint32 FParameter;
 
@@ -132,7 +133,7 @@ public:
 	FResolveSingleSamplePS(const ShaderMetaType::CompiledShaderInitializerType& Initializer);
 	
 	static bool ShouldCompilePermutation(const FGlobalShaderPermutationParameters& Parameters);
-	void SetParameters(FRHICommandList& RHICmdList, uint32 SingleSampleIndexValue);
+	RENDERCORE_API void SetParameters(FRHICommandList& RHICmdList, uint32 SingleSampleIndexValue);
 	
 	LAYOUT_FIELD(FShaderResourceParameter, UnresolvedSurface);
 	LAYOUT_FIELD(FShaderParameter, SingleSampleIndex);
@@ -141,23 +142,23 @@ public:
 /**
  * A vertex shader for rendering a textured screen element.
  */
-class RENDERCORE_API FResolveVS : public FGlobalShader
+class FResolveVS : public FGlobalShader
 {
-	DECLARE_GLOBAL_SHADER(FResolveVS);
+	DECLARE_EXPORTED_GLOBAL_SHADER(FResolveVS, RENDERCORE_API);
 public:
 	FResolveVS();
 	FResolveVS(const ShaderMetaType::CompiledShaderInitializerType& Initializer);
 	
 	static bool ShouldCompilePermutation(const FGlobalShaderPermutationParameters& Parameters);
-	void SetParameters(FRHICommandList& RHICmdList, const FResolveRect& SrcBounds, const FResolveRect& DstBounds, uint32 DstSurfaceWidth, uint32 DstSurfaceHeight);
+	RENDERCORE_API void SetParameters(FRHICommandList& RHICmdList, const FResolveRect& SrcBounds, const FResolveRect& DstBounds, uint32 DstSurfaceWidth, uint32 DstSurfaceHeight);
 
 	LAYOUT_FIELD(FShaderParameter, PositionMinMax);
 	LAYOUT_FIELD(FShaderParameter, UVMinMax);
 };
 
-class RENDERCORE_API FResolveArrayVS : public FResolveVS
+class FResolveArrayVS : public FResolveVS
 {
-	DECLARE_GLOBAL_SHADER(FResolveArrayVS);
+	DECLARE_EXPORTED_GLOBAL_SHADER(FResolveArrayVS, RENDERCORE_API);
 public:
 	FResolveArrayVS();
 	FResolveArrayVS(const ShaderMetaType::CompiledShaderInitializerType& Initializer);
@@ -165,5 +166,5 @@ public:
 	static bool ShouldCompilePermutation(const FGlobalShaderPermutationParameters& Parameters);
 	static void ModifyCompilationEnvironment(const FGlobalShaderPermutationParameters& Parameters, FShaderCompilerEnvironment& OutEnvironment);
 
-	void SetParameters(FRHICommandList& RHICmdList, const FResolveRect& SrcBounds, const FResolveRect& DstBounds, uint32 DstSurfaceWidth, uint32 DstSurfaceHeight);
+	RENDERCORE_API void SetParameters(FRHICommandList& RHICmdList, const FResolveRect& SrcBounds, const FResolveRect& DstBounds, uint32 DstSurfaceWidth, uint32 DstSurfaceHeight);
 };
diff --git a/Engine/Source/Runtime/Renderer/Private/RayTracing/RayTracingMaterialHitShaders.h b/Engine/Source/Runtime/Renderer/Private/RayTracing/RayTracingMaterialHitShaders.h
--- a/Engine/Source/Runtime/Renderer/Private/RayTracing/RayTracingMaterialHitShaders.h
+++ b/Engine/Source/Runtime/Renderer/Private/RayTracing/RayTracingMaterialHitShaders.h
@@ -121,9 +121,9 @@ private:
 	);
 };
 
-class RENDERER_API FHiddenMaterialHitGroup : public FGlobalShader
+class FHiddenMaterialHitGroup : public FGlobalShader
 {
-	DECLARE_GLOBAL_SHADER(FHiddenMaterialHitGroup)
+	DECLARE_EXPORTED_GLOBAL_SHADER(FHiddenMaterialHitGroup, RENDERER_API)
 	SHADER_USE_ROOT_PARAMETER_STRUCT(FHiddenMaterialHitGroup, FGlobalShader)
 
 	static bool ShouldCompilePermutation(const FGlobalShaderPermutationParameters& Parameters)
@@ -139,9 +139,9 @@ class RENDERER_API FHiddenMaterialHitGroup : public FGlobalShader
 	using FParameters = FEmptyShaderParameters;
 };
 
-class RENDERER_API FOpaqueShadowHitGroup : public FGlobalShader
+class FOpaqueShadowHitGroup : public FGlobalShader
 {
-	DECLARE_GLOBAL_SHADER(FOpaqueShadowHitGroup)
+	DECLARE_EXPORTED_GLOBAL_SHADER(FOpaqueShadowHitGroup, RENDERER_API)
 	SHADER_USE_ROOT_PARAMETER_STRUCT(FOpaqueShadowHitGroup, FGlobalShader)
 
 	static bool ShouldCompilePermutation(const FGlobalShaderPermutationParameters& Parameters)
@@ -157,9 +157,9 @@ class RENDERER_API FOpaqueShadowHitGroup : public FGlobalShader
 	using FParameters = FEmptyShaderParameters;
 };
 
-class RENDERER_API FDefaultCallableShader : public FGlobalShader
+class FDefaultCallableShader : public FGlobalShader
 {
-	DECLARE_GLOBAL_SHADER(FDefaultCallableShader)
+	DECLARE_EXPORTED_GLOBAL_SHADER(FDefaultCallableShader, RENDERER_API)
 	SHADER_USE_ROOT_PARAMETER_STRUCT(FDefaultCallableShader, FGlobalShader)
 
 	static bool ShouldCompilePermutation(const FGlobalShaderPermutationParameters& Parameters)
diff --git a/Engine/Source/Runtime/Renderer/Public/ScreenPass.h b/Engine/Source/Runtime/Renderer/Public/ScreenPass.h
--- a/Engine/Source/Runtime/Renderer/Public/ScreenPass.h
+++ b/Engine/Source/Runtime/Renderer/Public/ScreenPass.h
@@ -21,15 +21,10 @@ FRHITexture* GetMiniFontTexture();
 FRDGTextureRef RENDERER_API TryCreateViewFamilyTexture(FRDGBuilder& GraphBuilder, const FSceneViewFamily& ViewFamily);
 
 // The vertex shader used by DrawScreenPass to draw a rectangle.
-class RENDERER_API FScreenPassVS : public FGlobalShader
+class FScreenPassVS : public FGlobalShader
 {
 public:
-	DECLARE_GLOBAL_SHADER(FScreenPassVS);
-
-	static bool ShouldCompilePermutation(const FGlobalShaderPermutationParameters&)
-	{
-		return true;
-	}
+	DECLARE_EXPORTED_GLOBAL_SHADER(FScreenPassVS, RENDERER_API);
 
 	FScreenPassVS() = default;
 	FScreenPassVS(const ShaderMetaType::CompiledShaderInitializerType& Initializer)
diff --git a/Engine/Source/Runtime/Renderer/Public/VolumeRendering.h b/Engine/Source/Runtime/Renderer/Public/VolumeRendering.h
--- a/Engine/Source/Runtime/Renderer/Public/VolumeRendering.h
+++ b/Engine/Source/Runtime/Renderer/Public/VolumeRendering.h
@@ -44,9 +44,9 @@ struct FVolumeBounds
 };
 
 /** Vertex shader used to write to a range of slices of a 3d volume texture. */
-class ENGINE_API FWriteToSliceVS : public FGlobalShader
+class FWriteToSliceVS : public FGlobalShader
 {
-	DECLARE_GLOBAL_SHADER(FWriteToSliceVS);
+	DECLARE_EXPORTED_GLOBAL_SHADER(FWriteToSliceVS, ENGINE_API);
 public:
 	FWriteToSliceVS();
 	FWriteToSliceVS(const ShaderMetaType::CompiledShaderInitializerType& Initializer);
@@ -81,9 +81,9 @@ private:
 };
 
 /** Geometry shader used to write to a range of slices of a 3d volume texture. */
-class ENGINE_API FWriteToSliceGS : public FGlobalShader
+class FWriteToSliceGS : public FGlobalShader
 {
-	DECLARE_GLOBAL_SHADER(FWriteToSliceGS);
+	DECLARE_EXPORTED_GLOBAL_SHADER(FWriteToSliceGS, ENGINE_API);
 public:
 	FWriteToSliceGS();
 	FWriteToSliceGS(const ShaderMetaType::CompiledShaderInitializerType& Initializer);
