- Fixing rivermax genlock always failing on launch

diff --git a/Engine/Plugins/VirtualProduction/Rivermax/RivermaxCore/Source/RivermaxCore/Private/RivermaxInputStream.cpp b/Engine/Plugins/VirtualProduction/Rivermax/RivermaxCore/Source/RivermaxCore/Private/RivermaxInputStream.cpp
--- a/Engine/Plugins/VirtualProduction/Rivermax/RivermaxCore/Source/RivermaxCore/Private/RivermaxInputStream.cpp
+++ b/Engine/Plugins/VirtualProduction/Rivermax/RivermaxCore/Source/RivermaxCore/Private/RivermaxInputStream.cpp
@@ -155,7 +155,7 @@ namespace UE::RivermaxCore::Private
 	bool FRivermaxInputStream::Initialize(const FRivermaxInputStreamOptions& InOptions, IRivermaxInputStreamListener& InListener)
 	{
 		IRivermaxCoreModule& RivermaxModule = FModuleManager::LoadModuleChecked<IRivermaxCoreModule>(TEXT("RivermaxCore"));
-		if (RivermaxModule.GetRivermaxManager()->IsInitialized() == false)
+		if (RivermaxModule.GetRivermaxManager()->IsLibraryInitialized() == false)
 		{
 			UE_LOG(LogRivermax, Warning, TEXT("Can't create Rivermax Input Stream. Library isn't initialized."));
 			return false;
diff --git a/Engine/Plugins/VirtualProduction/Rivermax/RivermaxCore/Source/RivermaxCore/Private/RivermaxManager.cpp b/Engine/Plugins/VirtualProduction/Rivermax/RivermaxCore/Source/RivermaxCore/Private/RivermaxManager.cpp
--- a/Engine/Plugins/VirtualProduction/Rivermax/RivermaxCore/Source/RivermaxCore/Private/RivermaxManager.cpp
+++ b/Engine/Plugins/VirtualProduction/Rivermax/RivermaxCore/Source/RivermaxCore/Private/RivermaxManager.cpp
@@ -54,16 +54,8 @@ namespace UE::RivermaxCore::Private
 			FModuleManager::LoadModuleChecked<FCUDAModule>("CUDA");
 		}
 
-		const bool bIsLibraryLoaded = LoadRivermaxLibrary();
-		if (bIsLibraryLoaded && FApp::CanEverRender())
-		{
-			//Postpone initialization after all modules have been loaded
-			FCoreDelegates::OnAllModuleLoadingPhasesComplete.AddRaw(this, &FRivermaxManager::InitializeLibrary);
-		}
-		else
-		{
-			UE_LOG(LogRivermax, Log, TEXT("Skipping Rivermax initialization. Library won't be usable."));
-		}
+		//Postpone initialization after all modules have been loaded
+		FCoreDelegates::OnAllModuleLoadingPhasesComplete.AddRaw(this, &FRivermaxManager::InitializeLibrary);
 	}
 
 	FRivermaxManager::~FRivermaxManager()
@@ -73,7 +65,7 @@ namespace UE::RivermaxCore::Private
 			rmax_status_t Status = rmax_cleanup();
 			if (Status != RMAX_OK)
 			{
-				UE_LOG(LogRivermax, Log, TEXT("Failed to cleanup Rivermax Library. Status: %d"), Status);
+				UE_LOG(LogRivermax, Warning, TEXT("Failed to cleanup Rivermax Library. Status: %d"), Status);
 			}
 			UE_LOG(LogRivermax, Log, TEXT("Rivermax Library has shutdown."));
 		}
@@ -85,7 +77,7 @@ namespace UE::RivermaxCore::Private
 		}
 	}
 
-	bool FRivermaxManager::IsInitialized() const
+	bool FRivermaxManager::IsManagerInitialized() const
 	{
 		return bIsInitialized;
 	}
@@ -152,57 +144,68 @@ namespace UE::RivermaxCore::Private
 
 	void FRivermaxManager::InitializeLibrary()
 	{
-		rmax_init_config Config;
-		memset(&Config, 0, sizeof(Config));
-
-		rmax_status_t Status = rmax_init(&Config);
-		if (Status == RMAX_OK)
+		const bool bIsLibraryLoaded = LoadRivermaxLibrary();
+		if (bIsLibraryLoaded && FApp::CanEverRender())
 		{
-			bIsCleanupRequired = true;
+			rmax_init_config Config;
+			memset(&Config, 0, sizeof(Config));
 
-			const URivermaxSettings* Settings = GetDefault<URivermaxSettings>();
-			TimeSource = Settings->TimeSource;
-			bool bIsClockConfigured = InitializeClock(TimeSource);
-
-			if(!bIsClockConfigured && TimeSource == ERivermaxTimeSource::PTP)
+			rmax_status_t Status = rmax_init(&Config);
+			if (Status == RMAX_OK)
 			{
-				TimeSource = ERivermaxTimeSource::System;
-				bIsClockConfigured = InitializeClock(TimeSource);
-			}
+				bIsCleanupRequired = true;
 
-			if (bIsClockConfigured)
-			{
-				uint32 Major = 0;
-				uint32 Minor = 0;
-				uint32 Release = 0;
-				uint32 Build = 0;
-				Status = rmax_get_version(&Major, &Minor, &Release, &Build);
-				if (Status == RMAX_OK)
-				{
-					if (CVarRivermaxEnableGPUDirectCapability.GetValueOnGameThread())
-					{
-						VerifyGPUDirectCapability();
-					}
+				const URivermaxSettings* Settings = GetDefault<URivermaxSettings>();
+				TimeSource = Settings->TimeSource;
+				bool bIsClockConfigured = InitializeClock(TimeSource);
 
-					bIsInitialized = true;
-					const TCHAR* GPUDirectSupport = bIsGPUDirectSupported ? TEXT("with GPUDirect support.") : TEXT("without GPUDirect support.");
-					UE_LOG(LogRivermax, Log, TEXT("Rivermax library version %d.%d.%d.%d succesfully initialized, %s"), Major, Minor, Release, Build, GPUDirectSupport);
+				if (!bIsClockConfigured && TimeSource == ERivermaxTimeSource::PTP)
+				{
+					TimeSource = ERivermaxTimeSource::System;
+					bIsClockConfigured = InitializeClock(TimeSource);
 				}
-				else
+
+				if (bIsClockConfigured)
 				{
-					UE_LOG(LogRivermax, Log, TEXT("Failed to retrieve Rivermax library version. Status: %d"), Status);
+					uint32 Major = 0;
+					uint32 Minor = 0;
+					uint32 Release = 0;
+					uint32 Build = 0;
+					Status = rmax_get_version(&Major, &Minor, &Release, &Build);
+					if (Status == RMAX_OK)
+					{
+						if (CVarRivermaxEnableGPUDirectCapability.GetValueOnGameThread())
+						{
+							VerifyGPUDirectCapability();
+						}
+
+						bIsLibraryInitialized = true;
+						const TCHAR* GPUDirectSupport = bIsGPUDirectSupported ? TEXT("with GPUDirect support.") : TEXT("without GPUDirect support.");
+						UE_LOG(LogRivermax, Log, TEXT("Rivermax library version %d.%d.%d.%d succesfully initialized, %s"), Major, Minor, Release, Build, GPUDirectSupport);
+					}
+					else
+					{
+						UE_LOG(LogRivermax, Log, TEXT("Failed to retrieve Rivermax library version. Status: %d"), Status);
+					}
 				}
 			}
+			else if (Status == RMAX_ERR_LICENSE_ISSUE)
+			{
+				UE_LOG(LogRivermax, Warning, TEXT("Rivermax License could not be found. Have you configured RIVERMAX_LICENSE_PATH environment variable?"));
+			}
+
+			if (bIsLibraryInitialized == false)
+			{
+				UE_LOG(LogRivermax, Error, TEXT("Rivermax library failed to initialize. Status: %d"), Status);
+			}
 		}
-		else if (Status == RMAX_ERR_LICENSE_ISSUE)
+		else
 		{
-			UE_LOG(LogRivermax, Log, TEXT("Rivermax License could not be found. Have you configured RIVERMAX_LICENSE_PATH environment variable?"));
+			UE_LOG(LogRivermax, Log, TEXT("Skipping Rivermax initialization. Library won't be usable."));
 		}
 
-		if (bIsInitialized == false)
-		{
-			UE_LOG(LogRivermax, Log, TEXT("Rivermax library failed to initialize. Status: %d"), Status);
-		}
+		bIsInitialized = true;
+		PostInitDelegate.Broadcast();
 	}
 
 	TConstArrayView<UE::RivermaxCore::FRivermaxDeviceInfo> FRivermaxManager::GetDevices() const
@@ -361,6 +364,17 @@ namespace UE::RivermaxCore::Private
 			}
 		}
 	}
+
+	bool FRivermaxManager::IsLibraryInitialized() const
+	{
+		return bIsLibraryInitialized;
+	}
+
+	FOnPostRivermaxManagerInit& FRivermaxManager::OnPostRivermaxManagerInit()
+	{
+		return PostInitDelegate;
+	}
+
 }
 
 
diff --git a/Engine/Plugins/VirtualProduction/Rivermax/RivermaxCore/Source/RivermaxCore/Private/RivermaxManager.h b/Engine/Plugins/VirtualProduction/Rivermax/RivermaxCore/Source/RivermaxCore/Private/RivermaxManager.h
--- a/Engine/Plugins/VirtualProduction/Rivermax/RivermaxCore/Source/RivermaxCore/Private/RivermaxManager.h
+++ b/Engine/Plugins/VirtualProduction/Rivermax/RivermaxCore/Source/RivermaxCore/Private/RivermaxManager.h
@@ -16,7 +16,9 @@ namespace UE::RivermaxCore::Private
 
 	public:
 		//~ Begin IRivermaxManager interface
-		virtual bool IsInitialized() const override;
+		virtual bool IsManagerInitialized() const override;
+		virtual bool IsLibraryInitialized() const override;
+		virtual FOnPostRivermaxManagerInit& OnPostRivermaxManagerInit() override;
 		virtual uint64 GetTime() const override;
 		virtual ERivermaxTimeSource GetTimeSource() const override;
 		virtual TConstArrayView<FRivermaxDeviceInfo> GetDevices() const override;
@@ -33,9 +35,12 @@ namespace UE::RivermaxCore::Private
 
 	private:
 
-		/** Whether the library has been initialized and is usable */
+		/** True when manager has been initialized */
 		bool bIsInitialized = false;
 
+		/** True when library is usable */
+		bool bIsLibraryInitialized = false;
+
 		/** Whether library was initialized but failed along the way and needs to be cleaned up on exit */
 		bool bIsCleanupRequired = false;
 
@@ -48,7 +53,11 @@ namespace UE::RivermaxCore::Private
 		/** Current time source */
 		ERivermaxTimeSource TimeSource = ERivermaxTimeSource::PTP;
 
+		/** Rivermax device finder listing all usable interfaces */
 		TUniquePtr<FRivermaxDeviceFinder> DeviceFinder;
+
+		/** Delegate triggered when initialization has completed */
+		FOnPostRivermaxManagerInit PostInitDelegate;
 	};
 }
 
diff --git a/Engine/Plugins/VirtualProduction/Rivermax/RivermaxCore/Source/RivermaxCore/Private/RivermaxOutputStream.cpp b/Engine/Plugins/VirtualProduction/Rivermax/RivermaxCore/Source/RivermaxCore/Private/RivermaxOutputStream.cpp
--- a/Engine/Plugins/VirtualProduction/Rivermax/RivermaxCore/Source/RivermaxCore/Private/RivermaxOutputStream.cpp
+++ b/Engine/Plugins/VirtualProduction/Rivermax/RivermaxCore/Source/RivermaxCore/Private/RivermaxOutputStream.cpp
@@ -113,7 +113,7 @@ namespace UE::RivermaxCore::Private
 		TRACE_CPUPROFILER_EVENT_SCOPE(FRivermaxOutputStream::Initialize);
 
 		RivermaxModule = FModuleManager::GetModulePtr<IRivermaxCoreModule>(TEXT("RivermaxCore"));
-		if (RivermaxModule->GetRivermaxManager()->IsInitialized() == false)
+		if (RivermaxModule->GetRivermaxManager()->IsLibraryInitialized() == false)
 		{
 			UE_LOG(LogRivermax, Warning, TEXT("Can't create Rivermax Output Stream. Library isn't initialized."));
 			return false;
diff --git a/Engine/Plugins/VirtualProduction/Rivermax/RivermaxCore/Source/RivermaxCore/Public/IRivermaxManager.h b/Engine/Plugins/VirtualProduction/Rivermax/RivermaxCore/Source/RivermaxCore/Public/IRivermaxManager.h
--- a/Engine/Plugins/VirtualProduction/Rivermax/RivermaxCore/Source/RivermaxCore/Public/IRivermaxManager.h
+++ b/Engine/Plugins/VirtualProduction/Rivermax/RivermaxCore/Source/RivermaxCore/Public/IRivermaxManager.h
@@ -14,6 +14,9 @@ namespace UE::RivermaxCore
 		FString InterfaceAddress;
 	};
 
+	/** Delegate called when rivermax manager has been initialized and library tried to be loaded */
+	DECLARE_MULTICAST_DELEGATE(FOnPostRivermaxManagerInit)
+
 	/**
 	 * Not doing much at the moment but should be used as a central point to register every stream created
 	 * and keep track of stats. Also manages initialization of the library.
@@ -25,8 +28,14 @@ namespace UE::RivermaxCore
 
 	public:
 		
-		/** Returns true if Rivermax has been initialized and is usable */
-		virtual bool IsInitialized() const = 0;
+		/** Returns true if RivermaxManager has been initialized */
+		virtual bool IsManagerInitialized() const = 0;
+
+		/** Returns true if Rivermax library has been initialized and is usable */
+		virtual bool IsLibraryInitialized() const = 0;
+
+		/** Delegate triggered after manager has been initialized and library tried to be loaded */
+		virtual FOnPostRivermaxManagerInit& OnPostRivermaxManagerInit() = 0;
 
 		/** 
 		 * Returns current time in nanoseconds for Rivermax
diff --git a/Engine/Plugins/VirtualProduction/Rivermax/RivermaxMedia/Source/RivermaxMedia/Private/RivermaxCustomTimeStep.cpp b/Engine/Plugins/VirtualProduction/Rivermax/RivermaxMedia/Source/RivermaxMedia/Private/RivermaxCustomTimeStep.cpp
--- a/Engine/Plugins/VirtualProduction/Rivermax/RivermaxMedia/Source/RivermaxMedia/Private/RivermaxCustomTimeStep.cpp
+++ b/Engine/Plugins/VirtualProduction/Rivermax/RivermaxMedia/Source/RivermaxMedia/Private/RivermaxCustomTimeStep.cpp
@@ -21,26 +21,21 @@ bool URivermaxCustomTimeStep::Initialize(UEngine* InEngine)
 		State = ECustomTimeStepSynchronizationState::Error;
 		return false;
 	}
+	
+	State = ECustomTimeStepSynchronizationState::Synchronizing;
 
-	if (RivermaxModule->GetRivermaxManager()->IsInitialized() == false)
+	if (RivermaxModule->GetRivermaxManager()->IsManagerInitialized() == false)
 	{
-		UE_LOG(LogRivermaxMedia, Warning, TEXT("Can't initialize Rivermax custom timestep. Library isn't available."))
-		State = ECustomTimeStepSynchronizationState::Error;
-		return false;
+		RivermaxModule->GetRivermaxManager()->OnPostRivermaxManagerInit().AddUObject(this, &URivermaxCustomTimeStep::OnRivermaxManagerInitialized);
 	}
-
-	if (RivermaxModule->GetRivermaxManager()->GetDevices().Num() <= 0)
+	else
 	{
-		UE_LOG(LogRivermaxMedia, Warning, TEXT("Can't initialize Rivermax custom timestep. No devices were found."))
-		State = ECustomTimeStepSynchronizationState::Error;
-		return false;
+		OnRivermaxManagerInitialized();
 	}
 
-	State = ECustomTimeStepSynchronizationState::Synchronized;
 	return true;
 }
 
-
 void URivermaxCustomTimeStep::Shutdown(UEngine* InEngine)
 {
 #if WITH_EDITORONLY_DATA
@@ -49,6 +44,11 @@ void URivermaxCustomTimeStep::Shutdown(UEngine* InEngine)
 
 	State = ECustomTimeStepSynchronizationState::Closed;
 
+	IRivermaxCoreModule* RivermaxModule = FModuleManager::GetModulePtr<IRivermaxCoreModule>("RivermaxCore");
+	if (RivermaxModule)
+	{
+		RivermaxModule->GetRivermaxManager()->OnPostRivermaxManagerInit().RemoveAll(this);
+	}
 }
 
 bool URivermaxCustomTimeStep::UpdateTimeStep(UEngine* InEngine)
@@ -59,7 +59,7 @@ bool URivermaxCustomTimeStep::UpdateTimeStep(UEngine* InEngine)
 		State = ECustomTimeStepSynchronizationState::Error;
 	}
 
-	if (State == ECustomTimeStepSynchronizationState::Closed)
+	if (State == ECustomTimeStepSynchronizationState::Closed || State == ECustomTimeStepSynchronizationState::Synchronizing)
 	{
 		return true;
 	}
@@ -81,7 +81,6 @@ bool URivermaxCustomTimeStep::UpdateTimeStep(UEngine* InEngine)
 
 	if (!bWaitedForSync)
 	{
-		State = ECustomTimeStepSynchronizationState::Error;
 		return true;
 	}
 
@@ -213,4 +212,29 @@ bool URivermaxCustomTimeStep::WaitForNextFrame()
 	return ActualWaitTime < TimeoutSec;
 }
 
+void URivermaxCustomTimeStep::OnRivermaxManagerInitialized()
+{
+	if (State == ECustomTimeStepSynchronizationState::Synchronizing)
+	{
+		IRivermaxCoreModule& RivermaxModule = FModuleManager::GetModuleChecked<IRivermaxCoreModule>("RivermaxCore");
+
+		if (RivermaxModule.GetRivermaxManager()->IsLibraryInitialized() == false)
+		{
+			UE_LOG(LogRivermaxMedia, Warning, TEXT("Can't initialize Rivermax custom timestep. Library isn't available."));
+			
+			State = ECustomTimeStepSynchronizationState::Error;
+			return;
+		}
+
+		if (RivermaxModule.GetRivermaxManager()->GetDevices().Num() <= 0)
+		{
+			UE_LOG(LogRivermaxMedia, Warning, TEXT("Can't initialize Rivermax custom timestep. No devices were found."))
+			State = ECustomTimeStepSynchronizationState::Error;
+			return;
+		}
+
+		State = ECustomTimeStepSynchronizationState::Synchronized;
+	}
+}
+
 #undef LOCTEXT_NAMESPACE
diff --git a/Engine/Plugins/VirtualProduction/Rivermax/RivermaxMedia/Source/RivermaxMedia/Public/RivermaxCustomTimeStep.h b/Engine/Plugins/VirtualProduction/Rivermax/RivermaxMedia/Source/RivermaxMedia/Public/RivermaxCustomTimeStep.h
--- a/Engine/Plugins/VirtualProduction/Rivermax/RivermaxMedia/Source/RivermaxMedia/Public/RivermaxCustomTimeStep.h
+++ b/Engine/Plugins/VirtualProduction/Rivermax/RivermaxMedia/Source/RivermaxMedia/Public/RivermaxCustomTimeStep.h
@@ -39,8 +39,13 @@ public:
 
 
 private:
+
+	/** Waits for the next frame based on frame rate settings */
 	bool WaitForNextFrame();
 
+	/** We wait for Rivermax Manager to be initialized to do final verification before going to Synchronized state` */
+	void OnRivermaxManagerInitialized();
+
 public:
 
 	/** Target frame rate to which to genlock. Uses ST2059 standard to align PTP time to standard genlock */
