- Adding protection against crash during rivermax initialization when prerequesites are not met.
- Fixing a race condition preventing input stream to play

diff --git a/Engine/Plugins/VirtualProduction/Rivermax/RivermaxCore/Source/RivermaxCore/Private/RivermaxInputStream.cpp b/Engine/Plugins/VirtualProduction/Rivermax/RivermaxCore/Source/RivermaxCore/Private/RivermaxInputStream.cpp
--- a/Engine/Plugins/VirtualProduction/Rivermax/RivermaxCore/Source/RivermaxCore/Private/RivermaxInputStream.cpp
+++ b/Engine/Plugins/VirtualProduction/Rivermax/RivermaxCore/Source/RivermaxCore/Private/RivermaxInputStream.cpp
@@ -547,7 +547,7 @@ namespace UE::RivermaxCore::Private
 	void FRivermaxInputStream::AllocateBuffers()
 	{
 		IRivermaxCoreModule& RivermaxModule = FModuleManager::LoadModuleChecked<IRivermaxCoreModule>(TEXT("RivermaxCore"));
-		if (RivermaxModule.GetRivermaxManager()->IsGPUDirectSupported() && Options.bUseGPUDirect)
+		if (RivermaxModule.GetRivermaxManager()->IsGPUDirectInputSupported() && Options.bUseGPUDirect)
 		{
 			bIsUsingGPUDirect = AllocateGPUBuffers();
 		}
diff --git a/Engine/Plugins/VirtualProduction/Rivermax/RivermaxCore/Source/RivermaxCore/Private/RivermaxManager.cpp b/Engine/Plugins/VirtualProduction/Rivermax/RivermaxCore/Source/RivermaxCore/Private/RivermaxManager.cpp
--- a/Engine/Plugins/VirtualProduction/Rivermax/RivermaxCore/Source/RivermaxCore/Private/RivermaxManager.cpp
+++ b/Engine/Plugins/VirtualProduction/Rivermax/RivermaxCore/Source/RivermaxCore/Private/RivermaxManager.cpp
@@ -33,6 +33,16 @@ namespace UE::RivermaxCore::Private
 		TEXT("Whether GPUDirect capability is validated at launch. Can be used to disable it entirely."),
 		ECVF_Default);
 
+	static TAutoConsoleVariable<int32> CVarRivermaxEnableGPUDirectInputCapability(
+		TEXT("Rivermax.GPUDirectInput"), 0,
+		TEXT("Whether GPUDirect capability is enabled for input purposes."),
+		ECVF_Default);
+
+	static TAutoConsoleVariable<int32> CVarRivermaxEnableGPUDirectOutputCapability(
+		TEXT("Rivermax.GPUDirectOutput"), 1,
+		TEXT("Whether GPUDirect capability is enabled for output purposes."),
+		ECVF_Default);
+
 
 	static uint64 RivermaxTimeHandler(void* Context)
 	{
@@ -144,8 +154,13 @@ namespace UE::RivermaxCore::Private
 
 	void FRivermaxManager::InitializeLibrary()
 	{
-		const bool bIsLibraryLoaded = LoadRivermaxLibrary();
-		if (bIsLibraryLoaded && FApp::CanEverRender())
+		bool bCanProceed = VerifyPrerequesites();
+		if(bCanProceed)
+		{
+			bCanProceed = LoadRivermaxLibrary();
+		}
+		
+		if (bCanProceed && FApp::CanEverRender())
 		{
 			rmax_init_config Config;
 			memset(&Config, 0, sizeof(Config));
@@ -375,6 +390,58 @@ namespace UE::RivermaxCore::Private
 		return PostInitDelegate;
 	}
 
+	bool FRivermaxManager::VerifyPrerequesites()
+	{
+		// This registry key is required in order to use rivermax sdk. This puts the network adapter in rivermax mode.
+		// With 1.20, SDK currently crashes if that is not present / set. Avoid loading the library if that's the case and notify the user.
+		
+		bool bSuccess = false;
+		FString RivermaxSetupValue;
+		static const TCHAR* MLXFolder = TEXT("SYSTEM\\CurrentControlSet\\Services\\mlx5\\Parameters");
+		static const TCHAR* RivermaxKey = TEXT("RivermaxSetup");
+		const bool bFoundKey = FWindowsPlatformMisc::QueryRegKey(HKEY_LOCAL_MACHINE, MLXFolder, RivermaxKey, RivermaxSetupValue);
+		if (bFoundKey)
+		{
+			int32 KeyValue = 0;
+			LexFromString(KeyValue, *RivermaxSetupValue);
+			if (KeyValue == 1)
+			{
+				bSuccess = true;
+				UE_LOG(LogRivermax, Verbose, TEXT("Rivermax setup registry key was found and enabled."));
+			}
+			else
+			{
+				UE_LOG(LogRivermax, Warning, TEXT("Rivermax setup registry key was found but not enabled."));
+			}
+		}
+		else
+		{
+			UE_LOG(LogRivermax, Warning, TEXT("Could not find RivermaxSetup registry key. Verify installation and reboot your machine."));
+		}
+
+		if (bSuccess)
+		{
+			// Crash was reported calling rmax_init with all prerequesite installation correctly done but no card installed. 
+			// This will prevent us from going forward with rivermax initialization
+			bSuccess = DeviceFinder->GetDevices().Num() > 0;
+			if (!bSuccess)
+			{
+				UE_LOG(LogRivermax, Warning, TEXT("Could not find any compatible Rivermax network adapters."));
+			}
+		}
+
+		return bSuccess;
+	}
+
+	bool FRivermaxManager::IsGPUDirectInputSupported() const
+	{
+		return bIsGPUDirectSupported && CVarRivermaxEnableGPUDirectInputCapability.GetValueOnAnyThread();
+	}
+
+	bool FRivermaxManager::IsGPUDirectOutputSupported() const
+	{
+		return bIsGPUDirectSupported && CVarRivermaxEnableGPUDirectOutputCapability.GetValueOnAnyThread();
+	}
 }
 
 
diff --git a/Engine/Plugins/VirtualProduction/Rivermax/RivermaxCore/Source/RivermaxCore/Private/RivermaxManager.h b/Engine/Plugins/VirtualProduction/Rivermax/RivermaxCore/Source/RivermaxCore/Private/RivermaxManager.h
--- a/Engine/Plugins/VirtualProduction/Rivermax/RivermaxCore/Source/RivermaxCore/Private/RivermaxManager.h
+++ b/Engine/Plugins/VirtualProduction/Rivermax/RivermaxCore/Source/RivermaxCore/Private/RivermaxManager.h
@@ -25,14 +25,27 @@ namespace UE::RivermaxCore::Private
 		virtual bool GetMatchingDevice(const FString& InSourceIP, FString& OutDeviceIP) const override;
 		virtual bool IsValidIP(const FString& InSourceIP) const override;
 		virtual bool IsGPUDirectSupported() const;
+		virtual bool IsGPUDirectInputSupported() const;
+		virtual bool IsGPUDirectOutputSupported() const;
 		//~ End IRivermaxManager interface
 
 	private:
+
+		/** Loads the rivermax dll shipping with engine */
 		bool LoadRivermaxLibrary();
+		
+		/** Initializes rivermax library */
 		void InitializeLibrary();
+
+		/** Verify if system is capable of doing GPUDirect */
 		void VerifyGPUDirectCapability();
+
+		/** Initializes clock used by Rivermax internals */
 		bool InitializeClock(ERivermaxTimeSource DesiredTimeSource);
 
+		/** Verify if all prerequesites to go through Rivermax initialization have been met */
+		bool VerifyPrerequesites();
+
 	private:
 
 		/** True when manager has been initialized */
diff --git a/Engine/Plugins/VirtualProduction/Rivermax/RivermaxCore/Source/RivermaxCore/Private/RivermaxOutputStream.cpp b/Engine/Plugins/VirtualProduction/Rivermax/RivermaxCore/Source/RivermaxCore/Private/RivermaxOutputStream.cpp
--- a/Engine/Plugins/VirtualProduction/Rivermax/RivermaxCore/Source/RivermaxCore/Private/RivermaxOutputStream.cpp
+++ b/Engine/Plugins/VirtualProduction/Rivermax/RivermaxCore/Source/RivermaxCore/Private/RivermaxOutputStream.cpp
@@ -131,7 +131,7 @@ namespace UE::RivermaxCore::Private
 		}
 
 		// We do (try) to make gpu allocations here to let the capturer know if we require it or not.
-		if (RivermaxModule->GetRivermaxManager()->IsGPUDirectSupported() && Options.bUseGPUDirect)
+		if (RivermaxModule->GetRivermaxManager()->IsGPUDirectOutputSupported() && Options.bUseGPUDirect)
 		{
 			bUseGPUDirect = AllocateGPUBuffers();
 		}
diff --git a/Engine/Plugins/VirtualProduction/Rivermax/RivermaxCore/Source/RivermaxCore/Public/IRivermaxManager.h b/Engine/Plugins/VirtualProduction/Rivermax/RivermaxCore/Source/RivermaxCore/Public/IRivermaxManager.h
--- a/Engine/Plugins/VirtualProduction/Rivermax/RivermaxCore/Source/RivermaxCore/Public/IRivermaxManager.h
+++ b/Engine/Plugins/VirtualProduction/Rivermax/RivermaxCore/Source/RivermaxCore/Public/IRivermaxManager.h
@@ -58,6 +58,12 @@ namespace UE::RivermaxCore
 
 		/** Whether gpudirect is supported on this platform. At the moment, we rely on CUDA to handle it. */
 		virtual bool IsGPUDirectSupported() const = 0;
+
+		/** Whether gpudirect is enabled for input streams. */
+		virtual bool IsGPUDirectInputSupported() const = 0;
+
+		/** Whether gpudirect is enabled for output streams. */
+		virtual bool IsGPUDirectOutputSupported() const = 0;
 	};
 }
 
diff --git a/Engine/Plugins/VirtualProduction/Rivermax/RivermaxMedia/Source/RivermaxMedia/Private/RivermaxMediaPlayer.cpp b/Engine/Plugins/VirtualProduction/Rivermax/RivermaxMedia/Source/RivermaxMedia/Private/RivermaxMediaPlayer.cpp
--- a/Engine/Plugins/VirtualProduction/Rivermax/RivermaxMedia/Source/RivermaxMedia/Private/RivermaxMediaPlayer.cpp
+++ b/Engine/Plugins/VirtualProduction/Rivermax/RivermaxMedia/Source/RivermaxMedia/Private/RivermaxMediaPlayer.cpp
@@ -96,6 +96,10 @@ namespace UE::RivermaxMedia
 		{
 			InputStream = Module->CreateInputStream();
 		}
+
+		CurrentState = EMediaState::Preparing;
+		RivermaxThreadNewState = EMediaState::Preparing;
+		
 		if (InputStream == nullptr || !InputStream->Initialize(StreamOptions, *this))
 		{
 			UE_LOG(LogRivermaxMedia, Warning, TEXT("Failed to initialize Rivermax input stream."));
@@ -111,9 +115,6 @@ namespace UE::RivermaxMedia
 		// Allocate different buffers, samples we will be using
 		AllocateBuffers();
 
-		// finalize
-		CurrentState = EMediaState::Preparing;
-		RivermaxThreadNewState = EMediaState::Preparing;
 		EventSink.ReceiveMediaEvent(EMediaEvent::MediaConnecting);
 
 #if WITH_EDITOR
@@ -560,7 +561,7 @@ namespace UE::RivermaxMedia
 		{
 			// Provide the late sample source data for the converter
 			const TSharedPtr<FRivermaxMediaTextureSample>& NextSample = SamplePool[NextFrameExpectations.FrameIndex]->Sample;
-			if (StreamOptions.bUseGPUDirect)
+			if (bDoesStreamSupportsGPUDirect)
 			{
 				OutConverterSetup.GPUBuffer = NextSample->GetGPUBuffer();
 			}
