- Fixing warning during build DDC

diff --git a/Engine/Plugins/VirtualProduction/Rivermax/RivermaxCore/Source/RivermaxCore/Private/RivermaxManager.cpp b/Engine/Plugins/VirtualProduction/Rivermax/RivermaxCore/Source/RivermaxCore/Private/RivermaxManager.cpp
--- a/Engine/Plugins/VirtualProduction/Rivermax/RivermaxCore/Source/RivermaxCore/Private/RivermaxManager.cpp
+++ b/Engine/Plugins/VirtualProduction/Rivermax/RivermaxCore/Source/RivermaxCore/Private/RivermaxManager.cpp
@@ -154,13 +154,19 @@ namespace UE::RivermaxCore::Private
 
 	void FRivermaxManager::InitializeLibrary()
 	{
-		bool bCanProceed = VerifyPrerequesites();
+		bool bCanProceed = FApp::CanEverRender();
+		
+		if (bCanProceed)
+		{
+			bCanProceed = VerifyPrerequesites();
+		}
+
 		if(bCanProceed)
 		{
 			bCanProceed = LoadRivermaxLibrary();
 		}
 		
-		if (bCanProceed && FApp::CanEverRender())
+		if (bCanProceed)
 		{
 			rmax_init_config Config;
 			memset(&Config, 0, sizeof(Config));
