Fix to support adaptive overscan for ndisplay's inner frustum

- Adding support for multi SRD packets both on input and output aligned with the full frame size. Can be disabled through cvars.
- Adding option to enforce desired resoluiton on input and output. If not enforced, output will adapat itself to source texture size and input will adapt to incoming stream resolution.

diff --git a/Engine/Plugins/Media/MediaIOFramework/Source/MediaIOCore/Private/MediaCapture.cpp b/Engine/Plugins/Media/MediaIOFramework/Source/MediaIOCore/Private/MediaCapture.cpp
--- a/Engine/Plugins/Media/MediaIOFramework/Source/MediaIOCore/Private/MediaCapture.cpp
+++ b/Engine/Plugins/Media/MediaIOFramework/Source/MediaIOCore/Private/MediaCapture.cpp
@@ -71,6 +71,7 @@ DECLARE_GPU_STAT(MediaCapture_Readback);
 DECLARE_GPU_STAT(MediaCapture_ProcessCapture);
 DECLARE_GPU_STAT(MediaCapture_SyncPointPass);
 
+
 /* namespace MediaCaptureDetails definition
 *****************************************************************************/
 
@@ -482,18 +483,16 @@ namespace UE::MediaCaptureData
 				return false;
 			}
 
+			bool bFoundSizeMismatch = false;
+			FIntPoint RequestSize = FIntPoint::ZeroValue;
 			if (Args.MediaCapture->DesiredCaptureOptions.Crop == EMediaCaptureCroppingType::None)
 			{
 				if (Args.SourceViewRect.Area() != 0)
 				{
 					if (Args.DesiredSize.X != Args.SourceViewRect.Width() || Args.DesiredSize.Y != Args.SourceViewRect.Height())
 					{
-						UE_LOG(LogMediaIOCore, Error, TEXT("The capture will stop for '%s'. The Source size doesn't match with the user requested size. Requested: %d,%d  Source: %d,%d")
-							, *Args.MediaCapture->MediaOutputName
-							, Args.SourceViewRect.Width(), Args.SourceViewRect.Height()
-							, Args.GetSizeX(), Args.GetSizeY());
-
-						return false;
+						RequestSize = { Args.SourceViewRect.Width(), Args.SourceViewRect.Height() };
+						bFoundSizeMismatch = true;
 					}
 					else
 					{
@@ -504,12 +503,8 @@ namespace UE::MediaCaptureData
 				}
 				else if (Args.SourceViewRect.Area() == 0 && (Args.DesiredSize.X != Args.GetSizeX() || Args.DesiredSize.Y != Args.GetSizeY()))
 				{
-					UE_LOG(LogMediaIOCore, Error, TEXT("The capture will stop for '%s'. The Source size doesn't match with the user requested size. Requested: %d,%d  Source: %d,%d")
-						, *Args.MediaCapture->MediaOutputName
-						, Args.DesiredSize.X, Args.DesiredSize.Y
-						, Args.GetSizeX(), Args.GetSizeY());
-
-					return false;
+					RequestSize = { Args.DesiredSize };
+					bFoundSizeMismatch = true;
 				}
 			}
 			else
@@ -521,10 +516,32 @@ namespace UE::MediaCaptureData
 				}
 
 				if ((Args.DesiredSize.X + StartCapturePoint.X) > Args.GetSizeX() || (Args.DesiredSize.Y + StartCapturePoint.Y) > Args.GetSizeY())
+				{
+					RequestSize = { Args.DesiredSize };
+					bFoundSizeMismatch = true;
+				}
+			}
+
+			if (bFoundSizeMismatch)
+			{
+				if (Args.MediaCapture->SupportsAutoRestart() 
+					&& Args.MediaCapture->bUseRequestedTargetSize 
+					&& Args.MediaCapture->DesiredCaptureOptions.bAutoRestartOnSourceSizeChange)
+				{
+					Args.MediaCapture->bIsAutoRestartRequired = true;
+
+					UE_LOG(LogMediaIOCore, Log, TEXT("The capture will auto restart for '%s'. The Source size doesn't match with the user requested size. Requested: %d,%d  Source: %d,%d")
+						, *Args.MediaCapture->MediaOutputName
+						, RequestSize.X, RequestSize.Y
+						, Args.GetSizeX(), Args.GetSizeY());
+					
+					return false;
+				}
+				else
 				{
 					UE_LOG(LogMediaIOCore, Error, TEXT("The capture will stop for '%s'. The Source size doesn't match with the user requested size. Requested: %d,%d  Source: %d,%d")
 						, *Args.MediaCapture->MediaOutputName
-						, Args.DesiredSize.X, Args.DesiredSize.Y
+						, RequestSize.X, RequestSize.Y
 						, Args.GetSizeX(), Args.GetSizeY());
 
 					return false;
@@ -1053,6 +1070,7 @@ UMediaCapture::UMediaCapture()
 	, bShouldCaptureRHIResource(false)
 	, WaitingForRenderCommandExecutionCounter(0)
 	, PendingFrameCount(0)
+	, bIsAutoRestartRequired(false)
 {
 }
 
@@ -1527,6 +1545,11 @@ void UMediaCapture::CaptureImmediate_RenderThread(const UE::MediaCaptureData::FC
 	
 	check(IsInRenderingThread());
 
+	if (bIsAutoRestartRequired)
+	{
+		return;
+	}
+
 	// This could happen if the capture immediate is called before our resources were initialized. We can't really flush as we are in a command.
 	if (!bOutputResourcesInitialized)
 	{
@@ -1550,6 +1573,9 @@ void UMediaCapture::CaptureImmediate_RenderThread(const UE::MediaCaptureData::FC
 		SetState(EMediaCaptureState::Error);
 	}
 
+	// Keep resource size up to date with incoming resource to capture
+ 	StaticCastSharedPtr<UE::MediaCapture::Private::FRHIResourceCaptureSource>(CaptureSource)->ResourceDescription.ResourceSize = Args.GetSizeXY();
+
 	if (GetState() != EMediaCaptureState::Capturing && GetState() != EMediaCaptureState::StopRequested)
 	{
 		return;
@@ -1709,6 +1735,58 @@ void UMediaCapture::SetState(EMediaCaptureState InNewState)
 	}
 }
 
+void UMediaCapture::RestartCapture()
+{
+	check(IsInGameThread());
+
+	UE_LOG(LogMediaIOCore, Log, TEXT("Media Capture restarting for new size %dx%d"), CaptureSource->GetSize().X, CaptureSource->GetSize().Y);
+
+	if (FrameManager)
+	{
+		FrameManager->ForEachFrame([](const TSharedPtr<UE::MediaCaptureData::FFrame> InFrame)
+		{
+			StaticCastSharedPtr<UE::MediaCaptureData::FCaptureFrame>(InFrame)->bMediaCaptureActive = false;
+		});
+	}
+
+	constexpr bool bAllowPendingFrameToBeProcess = false;
+	if (GetState() != EMediaCaptureState::Stopped)
+	{
+		SetState(EMediaCaptureState::Preparing);
+
+		WaitForPendingTasks();
+
+		StopCaptureImpl(bAllowPendingFrameToBeProcess);
+
+		DesiredSize = CaptureSource->GetSize();
+		CacheOutputOptions();
+
+		bool bInitialized = InitializeCapture();
+		if (bInitialized)
+		{
+			bInitialized = CaptureSource->PostInitialize();
+		}
+
+		// This could have been updated by the initialization done by the implementation
+		bShouldCaptureRHIResource = ShouldCaptureRHIResource();
+
+		if (bInitialized)
+		{
+			bOutputResourcesInitialized = false;
+			InitializeOutputResources(MediaOutput->NumberOfTextureBuffers);
+			bInitialized = GetState() != EMediaCaptureState::Stopped;
+		}
+
+		if (!bInitialized)
+		{
+			SetState(EMediaCaptureState::Stopped);
+			MediaCaptureDetails::ShowSlateNotification();
+		}
+	}
+
+	bIsAutoRestartRequired = false;
+}
+
 void UMediaCapture::BroadcastStateChanged()
 {
 	OnStateChanged.Broadcast();
@@ -2147,6 +2225,11 @@ void UMediaCapture::OnEndFrame_GameThread()
 		FlushRenderingCommands();
 	}
 
+	if (bIsAutoRestartRequired)
+	{
+		return;
+	}
+
 	if (!MediaOutput)
 	{
 		return;
@@ -2234,7 +2317,22 @@ bool UMediaCapture::ProcessCapture_RenderThread(const TSharedPtr<UE::MediaCaptur
 		
 		if (bHasCaptureSuceeded == false)
 		{
-			Args.MediaCapture->SetState(EMediaCaptureState::Error);
+			if (Args.MediaCapture->bIsAutoRestartRequired)
+			{
+				TWeakObjectPtr<UMediaCapture> Self = this;
+				AsyncTask(ENamedThreads::GameThread, [Self]
+				{
+					UMediaCapture* MediaCapture = Self.Get();
+					if (UObjectInitialized() && MediaCapture)
+					{
+						MediaCapture->RestartCapture();
+					}
+				});
+			}
+			else
+			{
+				Args.MediaCapture->SetState(EMediaCaptureState::Error);
+			}
 		}
 
 		return bHasCaptureSuceeded;
diff --git a/Engine/Plugins/Media/MediaIOFramework/Source/MediaIOCore/Public/MediaCapture.h b/Engine/Plugins/Media/MediaIOFramework/Source/MediaIOCore/Public/MediaCapture.h
--- a/Engine/Plugins/Media/MediaIOFramework/Source/MediaIOCore/Public/MediaCapture.h
+++ b/Engine/Plugins/Media/MediaIOFramework/Source/MediaIOCore/Public/MediaCapture.h
@@ -144,6 +144,10 @@ struct MEDIAIOCORE_API FMediaCaptureOptions
 
 public:
 
+	/** If source texture's size change, auto restart capture to follow source resolution if implementation supports it. */
+	UPROPERTY(BlueprintReadWrite, EditAnywhere, Category = "MediaCapture")
+	bool bAutoRestartOnSourceSizeChange = false;
+
 	/** Action to do when game thread overruns render thread and all frames are in flights being captured / readback. */
 	UPROPERTY(BlueprintReadWrite, EditAnywhere, Category = "MediaCapture")
 	EMediaCaptureOverrunAction OverrunAction = EMediaCaptureOverrunAction::Flush;
@@ -398,6 +402,12 @@ protected:
 		return false;
 	}
 
+	/** Whether the capture implementation supports restarting on input format change. */
+	virtual bool SupportsAutoRestart() const
+	{
+		return false;
+	}
+
 	friend class UE::MediaCaptureData::FCaptureFrame;
 	friend class UE::MediaCaptureData::FMediaCaptureHelper;
 	struct FCaptureBaseData
@@ -529,7 +539,7 @@ protected:
 	EMediaCaptureConversionOperation GetConversionOperation() const { return ConversionOperation; }
 	FString GetCaptureSourceType() const;
 	void SetState(EMediaCaptureState InNewState);
-
+	void RestartCapture();
 
 private:
 	void InitializeOutputResources(int32 InNumberOfBuffers);
@@ -626,6 +636,7 @@ private:
 	std::atomic<int32> WaitingForRenderCommandExecutionCounter;
 	TArray<TFuture<void>> PendingCommands;
 	std::atomic<int32> PendingFrameCount;
+	std::atomic<bool> bIsAutoRestartRequired;
 
 	struct FQueuedCaptureData
 	{
diff --git a/Engine/Plugins/Runtime/nDisplay/Source/DisplayClusterMedia/Private/Capture/DisplayClusterMediaCaptureBase.cpp b/Engine/Plugins/Runtime/nDisplay/Source/DisplayClusterMedia/Private/Capture/DisplayClusterMediaCaptureBase.cpp
--- a/Engine/Plugins/Runtime/nDisplay/Source/DisplayClusterMedia/Private/Capture/DisplayClusterMediaCaptureBase.cpp
+++ b/Engine/Plugins/Runtime/nDisplay/Source/DisplayClusterMedia/Private/Capture/DisplayClusterMediaCaptureBase.cpp
@@ -148,6 +148,7 @@ bool FDisplayClusterMediaCaptureBase::StartMediaCapture()
 	
 	FMediaCaptureOptions MediaCaptureOptions;
 	MediaCaptureOptions.NumberOfFramesToCapture = -1;
+	MediaCaptureOptions.bAutoRestartOnSourceSizeChange = true;
 	MediaCaptureOptions.bSkipFrameWhenRunningExpensiveTasks = false;
 	MediaCaptureOptions.OverrunAction = EMediaCaptureOverrunAction::Flush;
 
diff --git a/Engine/Plugins/VirtualProduction/Rivermax/RivermaxCore/Shaders/Private/RivermaxShaders.usf b/Engine/Plugins/VirtualProduction/Rivermax/RivermaxCore/Shaders/Private/RivermaxShaders.usf
--- a/Engine/Plugins/VirtualProduction/Rivermax/RivermaxCore/Shaders/Private/RivermaxShaders.usf
+++ b/Engine/Plugins/VirtualProduction/Rivermax/RivermaxCore/Shaders/Private/RivermaxShaders.usf
@@ -16,9 +16,9 @@ float OnePixelDeltaX;
 float OnePixelDeltaY;
 float InputPixelOffsetX;
 float InputPixelOffsetY;
-uint HorizontalElementCount;
-uint VerticalElementCount;
-uint InputTexturePixelsPerThread;
+uint CapturedSizeX;
+uint OutputTextureSizeX;
+uint TotalElementCount;
 RWTexture2D<float4> OutTexture;
 
 struct FYUV8Bit422Buffer
@@ -73,128 +73,99 @@ StructuredBuffer<FRGB16fBuffer> InputRGB16fBuffer;
 
 // Shader to convert from RGB to YUV 422 8 bits
 
-[numthreads(8,8,1)]
+[numthreads(64,1,1)]
 void RGBToYUV8Bit422(uint3 ThreadId : SV_DispatchThreadID, uint GroupIndex : SV_GroupIndex)
 {
 	// Make sure it's a valid thread
-	if(ThreadId.x >= HorizontalElementCount || ThreadId.y >= VerticalElementCount)
+	if(ThreadId.x >= TotalElementCount)
 	{
 		return;
 	}
 
-	uint Index = ThreadId.y * HorizontalElementCount + ThreadId.x;
+	const uint Index = ThreadId.x;
 
-	float X = (ThreadId.x * InputTexturePixelsPerThread) + InputPixelOffsetX;
-	float Y = (ThreadId.y + InputPixelOffsetY) * OnePixelDeltaY;
+	const uint PixelCount = 2;
+	uint3 YUV[PixelCount];
 
-	// We process 2 RGB pixels per thread
-	float3 RGB0 = InputTexture.SampleLevel(InputSampler, float2(((X + 0.0) + 0.5) * OnePixelDeltaX, Y), 0).rgb;
-	float3 RGB1 = InputTexture.SampleLevel(InputSampler, float2(((X + 1.0) + 0.5) * OnePixelDeltaX, Y), 0).rgb;
-
-	// The 8 RGB needs to be converted to YUV. Each encoded pixel will use 2 RGB pixels (YnUnVnYn+1)
-	uint3 YUV0 = RgbToYuv(RGB0, ColorTransform, DoLinearToSrgb) * 255;
-	uint3 YUV1 = RgbToYuv(RGB1, ColorTransform, DoLinearToSrgb) * 255;
+	// The 2 RGB needs to be converted to YUV. Each encoded pixel will use 2 RGB pixels (YnUnVnYn+1)
+	const uint SourcePixelStart = ThreadId.x * PixelCount;
+	for (uint PixelIndex = 0; PixelIndex < PixelCount; PixelIndex++)
+	{
+		const float BaseX = ((SourcePixelStart + PixelIndex) % CapturedSizeX) + InputPixelOffsetX; 
+		const float BaseY = ((SourcePixelStart + PixelIndex) / CapturedSizeX) + InputPixelOffsetY;
+		const float U = (BaseX + 0.5) * OnePixelDeltaX;
+		const float V = (BaseY + 0.5) * OnePixelDeltaY;
+		const float3 RGB = InputTexture.SampleLevel(InputSampler, float2(U,V), 0).rgb;
+		YUV[PixelIndex] = RgbToYuv(RGB, ColorTransform, DoLinearToSrgb) * 255;
+	}
 
-	OutYUV4228bitBuffer[Index].DWord0 = (YUV1.x << 24) | (YUV0.z << 16) | (YUV0.x << 8) | (YUV0.y << 0);
+	OutYUV4228bitBuffer[Index].DWord0 = (YUV[1].x << 24) | (YUV[0].z << 16) | (YUV[0].x << 8) | (YUV[0].y << 0);
 }
 
 
 // Shader to convert from RGB to YUV 422 10 bits packed little endian
 
-[numthreads(8,8,1)]
+[numthreads(64,1,1)]
 void RGBToYUV10Bit422(uint3 ThreadId : SV_DispatchThreadID, uint GroupIndex : SV_GroupIndex)
 {
 	// Make sure it's a valid thread
-	if(ThreadId.x >= HorizontalElementCount || ThreadId.y >= VerticalElementCount)
+	if(ThreadId.x >= TotalElementCount)
 	{
 		return;
 	}
 
-	uint Index = ThreadId.y * HorizontalElementCount + ThreadId.x;
-
-	float X = (ThreadId.x * InputTexturePixelsPerThread) + InputPixelOffsetX;
-	float Y = (ThreadId.y + InputPixelOffsetY) * OnePixelDeltaY;
+	const uint Index = ThreadId.x;
 
-	// We process 8 RGB pixels per thread
-	float3 RGB0 = InputTexture.SampleLevel(InputSampler, float2(((X + 0.0) + 0.5) * OnePixelDeltaX, Y), 0).rgb;
-	float3 RGB1 = InputTexture.SampleLevel(InputSampler, float2(((X + 1.0) + 0.5) * OnePixelDeltaX, Y), 0).rgb;
-	float3 RGB2 = InputTexture.SampleLevel(InputSampler, float2(((X + 2.0) + 0.5) * OnePixelDeltaX, Y), 0).rgb;
-	float3 RGB3 = InputTexture.SampleLevel(InputSampler, float2(((X + 3.0) + 0.5) * OnePixelDeltaX, Y), 0).rgb;
-	float3 RGB4 = InputTexture.SampleLevel(InputSampler, float2(((X + 4.0) + 0.5) * OnePixelDeltaX, Y), 0).rgb;
-	float3 RGB5 = InputTexture.SampleLevel(InputSampler, float2(((X + 5.0) + 0.5) * OnePixelDeltaX, Y), 0).rgb;
-	float3 RGB6 = InputTexture.SampleLevel(InputSampler, float2(((X + 6.0) + 0.5) * OnePixelDeltaX, Y), 0).rgb;
-	float3 RGB7 = InputTexture.SampleLevel(InputSampler, float2(((X + 7.0) + 0.5) * OnePixelDeltaX, Y), 0).rgb;
+	const uint PixelCount = 8;
+	uint3 YUV[PixelCount];
 
 	// The 8 RGB needs to be converted to YUV. Each encoded pixel will use 2 RGB pixels (YnUnVnYn+1)
-	uint3 YUV0 = RgbToYuv(RGB0, ColorTransform, DoLinearToSrgb) * 1023;
-	uint3 YUV1 = RgbToYuv(RGB1, ColorTransform, DoLinearToSrgb) * 1023;
-	uint3 YUV2 = RgbToYuv(RGB2, ColorTransform, DoLinearToSrgb) * 1023;
-	uint3 YUV3 = RgbToYuv(RGB3, ColorTransform, DoLinearToSrgb) * 1023;
-	uint3 YUV4 = RgbToYuv(RGB4, ColorTransform, DoLinearToSrgb) * 1023;
-	uint3 YUV5 = RgbToYuv(RGB5, ColorTransform, DoLinearToSrgb) * 1023;
-	uint3 YUV6 = RgbToYuv(RGB6, ColorTransform, DoLinearToSrgb) * 1023;
-	uint3 YUV7 = RgbToYuv(RGB7, ColorTransform, DoLinearToSrgb) * 1023;
+	const uint SourcePixelStart = ThreadId.x * PixelCount;
+	for (uint PixelIndex = 0; PixelIndex < PixelCount; PixelIndex++)
+	{
+		const float BaseX = ((SourcePixelStart + PixelIndex) % CapturedSizeX) + InputPixelOffsetX; 
+		const float BaseY = ((SourcePixelStart + PixelIndex) / CapturedSizeX) + InputPixelOffsetY;
+		const float U = (BaseX + 0.5) * OnePixelDeltaX;
+		const float V = (BaseY + 0.5) * OnePixelDeltaY;
+		const float3 RGB = InputTexture.SampleLevel(InputSampler, float2(U,V), 0).rgb;
+		YUV[PixelIndex] = RgbToYuv(RGB, ColorTransform, DoLinearToSrgb) * 1023;
+	}
 
 	// We write the 4 YUV pixels (40 bits * 4 = 160 bits = 5x32 bits) following the little endian packing 2110 expects
-	OutYCbCrBuffer[Index].DWord0 = ((YUV0.z & 0x3F) << 26) | ((YUV1.x >> 8) << 24) | ((YUV0.x & 0xF) << 20) | ((YUV0.z >> 6) << 16) | ((YUV0.y & 0x3) << 14) | ((YUV0.x >> 4) << 8) | (YUV0.y >> 2);
-	OutYCbCrBuffer[Index].DWord1 = ((YUV2.x & 0xF) << 28) | ((YUV2.z >> 6) << 24) | ((YUV2.y & 0x3) << 22) | ((YUV2.x >> 4) << 16) | ((YUV2.y >> 2) << 8) | (YUV1.x & 0xFF);
-	OutYCbCrBuffer[Index].DWord2 = ((YUV4.y & 0x3) << 30) | ((YUV4.x >> 4) << 24) | ((YUV4.y >> 2) << 16) | ((YUV3.x & 0xFF) << 8) | ((YUV2.z & 0x3F) << 2) | (YUV3.x >> 8);
-	OutYCbCrBuffer[Index].DWord3 = ((YUV6.y >> 2) << 24) | ((YUV5.x & 0xFF) << 16) | ((YUV4.z & 0x3F) << 10) | ((YUV5.x >> 8) << 8) | ((YUV4.x & 0xF) << 4) | (YUV4.z >> 6);
-	OutYCbCrBuffer[Index].DWord4 = ((YUV7.x & 0xFF) << 24) | ((YUV6.z & 0x3F) << 18) | ((YUV7.x >> 8) << 16) | ((YUV6.x & 0xF) << 12) | ((YUV6.z >> 6) << 8) | ((YUV6.y & 0x3) << 6) | (YUV6.x >> 4);
+	OutYCbCrBuffer[Index].DWord0 = ((YUV[0].z & 0x3F) << 26) | ((YUV[1].x >> 8) << 24) | ((YUV[0].x & 0xF) << 20) | ((YUV[0].z >> 6) << 16) | ((YUV[0].y & 0x3) << 14) | ((YUV[0].x >> 4) << 8) | (YUV[0].y >> 2);
+	OutYCbCrBuffer[Index].DWord1 = ((YUV[2].x & 0xF) << 28) | ((YUV[2].z >> 6) << 24) | ((YUV[2].y & 0x3) << 22) | ((YUV[2].x >> 4) << 16) | ((YUV[2].y >> 2) << 8) | (YUV[1].x & 0xFF);
+	OutYCbCrBuffer[Index].DWord2 = ((YUV[4].y & 0x3) << 30) | ((YUV[4].x >> 4) << 24) | ((YUV[4].y >> 2) << 16) | ((YUV[3].x & 0xFF) << 8) | ((YUV[2].z & 0x3F) << 2) | (YUV[3].x >> 8);
+	OutYCbCrBuffer[Index].DWord3 = ((YUV[6].y >> 2) << 24) | ((YUV[5].x & 0xFF) << 16) | ((YUV[4].z & 0x3F) << 10) | ((YUV[5].x >> 8) << 8) | ((YUV[4].x & 0xF) << 4) | (YUV[4].z >> 6);
+	OutYCbCrBuffer[Index].DWord4 = ((YUV[7].x & 0xFF) << 24) | ((YUV[6].z & 0x3F) << 18) | ((YUV[7].x >> 8) << 16) | ((YUV[6].x & 0xF) << 12) | ((YUV[6].z >> 6) << 8) | ((YUV[6].y & 0x3) << 6) | (YUV[6].x >> 4);
 }
 
 // Shader to convert from RGBA to RGB10
 
-[numthreads(8,8,1)]
+[numthreads(64,1,1)]
 void RGBToRGB10BitCS(uint3 ThreadId : SV_DispatchThreadID, uint GroupIndex : SV_GroupIndex)
 {
 	// Make sure it's a valid thread
-	if(ThreadId.x >= HorizontalElementCount || ThreadId.y >= VerticalElementCount)
+	if(ThreadId.x >= TotalElementCount)
 	{
 		return;
 	}
 
-	uint Index = ThreadId.y * HorizontalElementCount + ThreadId.x;
-
-	float X = (ThreadId.x * InputTexturePixelsPerThread) + InputPixelOffsetX;
-	float Y = (ThreadId.y + InputPixelOffsetY) * OnePixelDeltaY;
-
-	uint Denormalizer = 1023;
-	uint3 RGB00 = InputTexture.SampleLevel(InputSampler, float2(((X + 0.0) + 0.5) * OnePixelDeltaX, Y), 0).rgb * Denormalizer;
-	uint3 RGB01 = InputTexture.SampleLevel(InputSampler, float2(((X + 1.0) + 0.5) * OnePixelDeltaX, Y), 0).rgb * Denormalizer;
-	uint3 RGB02 = InputTexture.SampleLevel(InputSampler, float2(((X + 2.0) + 0.5) * OnePixelDeltaX, Y), 0).rgb * Denormalizer;
-	uint3 RGB03 = InputTexture.SampleLevel(InputSampler, float2(((X + 3.0) + 0.5) * OnePixelDeltaX, Y), 0).rgb * Denormalizer;
-	uint3 RGB04 = InputTexture.SampleLevel(InputSampler, float2(((X + 4.0) + 0.5) * OnePixelDeltaX, Y), 0).rgb * Denormalizer;
-	uint3 RGB05 = InputTexture.SampleLevel(InputSampler, float2(((X + 5.0) + 0.5) * OnePixelDeltaX, Y), 0).rgb * Denormalizer;
-	uint3 RGB06 = InputTexture.SampleLevel(InputSampler, float2(((X + 6.0) + 0.5) * OnePixelDeltaX, Y), 0).rgb * Denormalizer;
-	uint3 RGB07 = InputTexture.SampleLevel(InputSampler, float2(((X + 7.0) + 0.5) * OnePixelDeltaX, Y), 0).rgb * Denormalizer;
-	uint3 RGB08 = InputTexture.SampleLevel(InputSampler, float2(((X + 8.0) + 0.5) * OnePixelDeltaX, Y), 0).rgb * Denormalizer;
-	uint3 RGB09 = InputTexture.SampleLevel(InputSampler, float2(((X + 9.0) + 0.5) * OnePixelDeltaX, Y), 0).rgb * Denormalizer;
-	uint3 RGB10 = InputTexture.SampleLevel(InputSampler, float2(((X + 10.0) + 0.5) * OnePixelDeltaX, Y), 0).rgb * Denormalizer;
-	uint3 RGB11 = InputTexture.SampleLevel(InputSampler, float2(((X + 11.0) + 0.5) * OnePixelDeltaX, Y), 0).rgb * Denormalizer;
-	uint3 RGB12 = InputTexture.SampleLevel(InputSampler, float2(((X + 12.0) + 0.5) * OnePixelDeltaX, Y), 0).rgb * Denormalizer;
-	uint3 RGB13 = InputTexture.SampleLevel(InputSampler, float2(((X + 13.0) + 0.5) * OnePixelDeltaX, Y), 0).rgb * Denormalizer;
-	uint3 RGB14 = InputTexture.SampleLevel(InputSampler, float2(((X + 14.0) + 0.5) * OnePixelDeltaX, Y), 0).rgb * Denormalizer;
-	uint3 RGB15 = InputTexture.SampleLevel(InputSampler, float2(((X + 15.0) + 0.5) * OnePixelDeltaX, Y), 0).rgb * Denormalizer;
-
-	// Cap values to be in range
-	RGB00 = min(RGB00, Denormalizer);
-	RGB01 = min(RGB01, Denormalizer);
-	RGB02 = min(RGB02, Denormalizer);
-	RGB03 = min(RGB03, Denormalizer);
-	RGB04 = min(RGB04, Denormalizer);
-	RGB05 = min(RGB05, Denormalizer);
-	RGB06 = min(RGB06, Denormalizer);
-	RGB07 = min(RGB07, Denormalizer);
-	RGB08 = min(RGB08, Denormalizer);
-	RGB09 = min(RGB09, Denormalizer);
-	RGB10 = min(RGB10, Denormalizer);
-	RGB11 = min(RGB11, Denormalizer);
-	RGB12 = min(RGB12, Denormalizer);
-	RGB13 = min(RGB13, Denormalizer);
-	RGB14 = min(RGB14, Denormalizer);
-	RGB15 = min(RGB15, Denormalizer);
-
+	// Convert output buffer index to pixels
+	const uint PixelCount = 16;
+	const uint SourcePixelStart = (ThreadId.x * PixelCount);
+	const uint Denormalizer = 1023;
+	const uint Index = ThreadId.x;
+	
+	uint3 RGB[PixelCount];
+	for (uint PixelIndex = 0; PixelIndex < PixelCount; PixelIndex++)
+	{
+		const float BaseX = ((SourcePixelStart + PixelIndex) % CapturedSizeX) + InputPixelOffsetX; 
+		const float BaseY = ((SourcePixelStart + PixelIndex) / CapturedSizeX) + InputPixelOffsetY;
+		const float U = (BaseX + 0.5) * OnePixelDeltaX;
+		const float V = (BaseY + 0.5) * OnePixelDeltaY;
+		RGB[PixelIndex] = InputTexture.SampleLevel(InputSampler, float2(U,V), 0).rgb * Denormalizer; 
+		RGB[PixelIndex] = min(RGB[PixelIndex], Denormalizer);
+	}
 
 	// All this bit swizzling is to pack 3 channels of 10 bits with R being in LSB of the first 32 bits and subsequent channels following until we used all 16 pixels.
 	// Example : RGB = (0x123, 0x245, 0x0AB)
@@ -206,32 +177,32 @@ void RGBToRGB10BitCS(uint3 ThreadId : SV_DispatchThreadID, uint GroupIndex : SV_
 	// We want 48 at the first adress so it goes in LSB of the 32 bits as it's written as little endian. 
 	// 2110 on the wire will read it as big endian which will give us R going out first
 
-	OutRGB10Buffer[Index].DWords[0] = ((RGB00.z & 0x03F) << 26) | ((RGB01.x >> 8)   << 24)   | ((RGB00.y & 0xF)  << 20) | ((RGB00.z >> 6)   << 16) | ((RGB00.x & 0x3)  << 14) | ((RGB00.y >> 4)   << 8)								| (RGB00.x >> 2);
-	OutRGB10Buffer[Index].DWords[1] = ((RGB01.z & 0x00F) << 28) | ((RGB02.x >> 6)   << 24)   | ((RGB01.y & 0x3)  << 22) | ((RGB01.z >> 4)   << 16)							  | ((RGB01.y >> 2)   << 8)								| (RGB01.x & 0xFF);
-	OutRGB10Buffer[Index].DWords[2] = ((RGB02.z & 0x003) << 30) | ((RGB03.x >> 4)   << 24)								| ((RGB02.z >> 2)   << 16)							  | ((RGB02.y & 0xFF) << 8)  | ((RGB02.x & 0x3F) << 2)  | (RGB02.y >> 8);
-	OutRGB10Buffer[Index].DWords[3] =							  ((RGB04.x >> 2)   << 24)								| ((RGB03.z & 0xFF) << 16) | ((RGB03.y & 0x3F) << 10) | ((RGB03.z >> 8)   << 8)  | ((RGB03.x & 0xF)  << 4)  | (RGB03.y >> 6);
-	OutRGB10Buffer[Index].DWords[4] =							  ((RGB05.x & 0xFF) << 24)   | ((RGB04.z & 0x3F) << 18) | ((RGB05.x >> 8)   << 16) | ((RGB04.y & 0xF)  << 12) | ((RGB04.z >> 6)   << 8)  | ((RGB04.x & 0x3)  << 6)  | (RGB04.y >> 4);
+	OutRGB10Buffer[Index].DWords[0] = ((RGB[0].z & 0x03F) << 26) | ((RGB[1].x >> 8)   << 24)   | ((RGB[0].y & 0xF)  << 20) | ((RGB[0].z >> 6)   << 16) | ((RGB[0].x & 0x3)  << 14) | ((RGB[0].y >> 4)   << 8)								| (RGB[0].x >> 2);
+	OutRGB10Buffer[Index].DWords[1] = ((RGB[1].z & 0x00F) << 28) | ((RGB[2].x >> 6)   << 24)   | ((RGB[1].y & 0x3)  << 22) | ((RGB[1].z >> 4)   << 16)							  | ((RGB[1].y >> 2)   << 8)								| (RGB[1].x & 0xFF);
+	OutRGB10Buffer[Index].DWords[2] = ((RGB[2].z & 0x003) << 30) | ((RGB[3].x >> 4)   << 24)								| ((RGB[2].z >> 2)   << 16)							  | ((RGB[2].y & 0xFF) << 8)  | ((RGB[2].x & 0x3F) << 2)  | (RGB[2].y >> 8);
+	OutRGB10Buffer[Index].DWords[3] =							  ((RGB[4].x >> 2)   << 24)								| ((RGB[3].z & 0xFF) << 16) | ((RGB[3].y & 0x3F) << 10) | ((RGB[3].z >> 8)   << 8)  | ((RGB[3].x & 0xF)  << 4)  | (RGB[3].y >> 6);
+	OutRGB10Buffer[Index].DWords[4] =							  ((RGB[5].x & 0xFF) << 24)   | ((RGB[4].z & 0x3F) << 18) | ((RGB[5].x >> 8)   << 16) | ((RGB[4].y & 0xF)  << 12) | ((RGB[4].z >> 6)   << 8)  | ((RGB[4].x & 0x3)  << 6)  | (RGB[4].y >> 4);
 	
-	OutRGB10Buffer[Index].DWords[5] = ((RGB06.x & 0x03F) << 26) | ((RGB06.y >> 8)   << 24)   | ((RGB05.z & 0xF)  << 20) | ((RGB06.x >> 6)   << 16) | ((RGB05.y & 0x3)  << 14) | ((RGB05.z >> 4)   << 8)								| (RGB05.y >> 2);
-	OutRGB10Buffer[Index].DWords[6] = ((RGB07.x & 0x00F) << 28) | ((RGB07.y >> 6)   << 24)   | ((RGB06.z & 0x3)  << 22) | ((RGB07.x >> 4)   << 16)							  | ((RGB06.z >> 2)   << 8)								| (RGB06.y & 0xFF);
-	OutRGB10Buffer[Index].DWords[7] = ((RGB08.x & 0x003) << 30) | ((RGB08.y >> 4)   << 24)								| ((RGB08.x >> 2)   << 16)							  | ((RGB07.z & 0xFF) << 8)  | ((RGB07.y & 0x3F) << 2)  | (RGB07.z >> 8);
-	OutRGB10Buffer[Index].DWords[8] =							  ((RGB09.y >> 2)   << 24)								| ((RGB09.x & 0xFF) << 16) | ((RGB08.z & 0x3F) << 10) | ((RGB09.x >> 8)   << 8)  | ((RGB08.y & 0xF)  << 4)  | (RGB08.z >> 6);
-	OutRGB10Buffer[Index].DWords[9] =							  ((RGB10.y & 0xFF) << 24)   | ((RGB10.x & 0x3F) << 18) | ((RGB10.y >> 8)   << 16) | ((RGB09.z & 0xF)  << 12) | ((RGB10.x >> 6)   << 8)  | ((RGB09.y & 0x3)  << 6)  | (RGB09.z >> 4);
-
-	OutRGB10Buffer[Index].DWords[10] = ((RGB11.y & 0x03F) << 26) | ((RGB11.z >> 8)   << 24)  | ((RGB11.x & 0xF)  << 20) | ((RGB11.y >> 6)   << 16) | ((RGB10.z & 0x3)  << 14) | ((RGB11.x >> 4)   << 8)								| (RGB10.z >> 2);
-	OutRGB10Buffer[Index].DWords[11] = ((RGB12.y & 0x00F) << 28) | ((RGB12.z >> 6)   << 24)  | ((RGB12.x & 0x3)  << 22) | ((RGB12.y >> 4)   << 16)							  | ((RGB12.x >> 2)   << 8)								| (RGB11.z & 0xFF);
-	OutRGB10Buffer[Index].DWords[12] = ((RGB13.y & 0x003) << 30) | ((RGB13.z >> 4)   << 24)								| ((RGB13.y >> 2)   << 16)							  | ((RGB13.x & 0xFF) << 8)  | ((RGB12.z & 0x3F) << 2)  | (RGB13.x >> 8);
-	OutRGB10Buffer[Index].DWords[13] =							   ((RGB14.z >> 2)   << 24)								| ((RGB14.y & 0xFF) << 16) | ((RGB14.x & 0x3F) << 10) | ((RGB14.y >> 8)   << 8)  | ((RGB13.z & 0xF)  << 4)  | (RGB14.x >> 6);
-	OutRGB10Buffer[Index].DWords[14] =							   ((RGB15.z & 0xFF) << 24)  | ((RGB15.y & 0x3F) << 18) | ((RGB15.z >> 8)   << 16) | ((RGB15.x & 0xF)  << 12) | ((RGB15.y >> 6)   << 8)  | ((RGB14.z & 0x3)  << 6)  | (RGB15.x >> 4);
+	OutRGB10Buffer[Index].DWords[5] = ((RGB[6].x & 0x03F) << 26) | ((RGB[6].y >> 8)   << 24)   | ((RGB[5].z & 0xF)  << 20) | ((RGB[6].x >> 6)   << 16) | ((RGB[5].y & 0x3)  << 14) | ((RGB[5].z >> 4)   << 8)								| (RGB[5].y >> 2);
+	OutRGB10Buffer[Index].DWords[6] = ((RGB[7].x & 0x00F) << 28) | ((RGB[7].y >> 6)   << 24)   | ((RGB[6].z & 0x3)  << 22) | ((RGB[7].x >> 4)   << 16)							  | ((RGB[6].z >> 2)   << 8)								| (RGB[6].y & 0xFF);
+	OutRGB10Buffer[Index].DWords[7] = ((RGB[8].x & 0x003) << 30) | ((RGB[8].y >> 4)   << 24)								| ((RGB[8].x >> 2)   << 16)							  | ((RGB[7].z & 0xFF) << 8)  | ((RGB[7].y & 0x3F) << 2)  | (RGB[7].z >> 8);
+	OutRGB10Buffer[Index].DWords[8] =							  ((RGB[9].y >> 2)   << 24)								| ((RGB[9].x & 0xFF) << 16) | ((RGB[8].z & 0x3F) << 10) | ((RGB[9].x >> 8)   << 8)  | ((RGB[8].y & 0xF)  << 4)  | (RGB[8].z >> 6);
+	OutRGB10Buffer[Index].DWords[9] =							  ((RGB[10].y & 0xFF) << 24)   | ((RGB[10].x & 0x3F) << 18) | ((RGB[10].y >> 8)   << 16) | ((RGB[9].z & 0xF)  << 12) | ((RGB[10].x >> 6)   << 8)  | ((RGB[9].y & 0x3)  << 6)  | (RGB[9].z >> 4);
+
+	OutRGB10Buffer[Index].DWords[10] = ((RGB[11].y & 0x03F) << 26) | ((RGB[11].z >> 8)   << 24)  | ((RGB[11].x & 0xF)  << 20) | ((RGB[11].y >> 6)   << 16) | ((RGB[10].z & 0x3)  << 14) | ((RGB[11].x >> 4)   << 8)								| (RGB[10].z >> 2);
+	OutRGB10Buffer[Index].DWords[11] = ((RGB[12].y & 0x00F) << 28) | ((RGB[12].z >> 6)   << 24)  | ((RGB[12].x & 0x3)  << 22) | ((RGB[12].y >> 4)   << 16)							  | ((RGB[12].x >> 2)   << 8)								| (RGB[11].z & 0xFF);
+	OutRGB10Buffer[Index].DWords[12] = ((RGB[13].y & 0x003) << 30) | ((RGB[13].z >> 4)   << 24)								| ((RGB[13].y >> 2)   << 16)							  | ((RGB[13].x & 0xFF) << 8)  | ((RGB[12].z & 0x3F) << 2)  | (RGB[13].x >> 8);
+	OutRGB10Buffer[Index].DWords[13] =							   ((RGB[14].z >> 2)   << 24)								| ((RGB[14].y & 0xFF) << 16) | ((RGB[14].x & 0x3F) << 10) | ((RGB[14].y >> 8)   << 8)  | ((RGB[13].z & 0xF)  << 4)  | (RGB[14].x >> 6);
+	OutRGB10Buffer[Index].DWords[14] =							   ((RGB[15].z & 0xFF) << 24)  | ((RGB[15].y & 0x3F) << 18) | ((RGB[15].z >> 8)   << 16) | ((RGB[15].x & 0xF)  << 12) | ((RGB[15].y >> 6)   << 8)  | ((RGB[14].z & 0x3)  << 6)  | (RGB[15].x >> 4);
 }
 
 // Shader to convert from packed RGB10 to RGB10A2 texture
 
-[numthreads(8,8,1)]
+[numthreads(64,1,1)]
 void RGB10BitToRGBACS(uint3 ThreadId : SV_DispatchThreadID, uint GroupIndex : SV_GroupIndex)
 {
 	// Make sure it's a valid thread
-	if(ThreadId.x >= HorizontalElementCount || ThreadId.y >= VerticalElementCount)
+	if(ThreadId.x >= TotalElementCount)
 	{
 		return;
 	}
@@ -239,401 +210,341 @@ void RGB10BitToRGBACS(uint3 ThreadId : SV_DispatchThreadID, uint GroupIndex : SV
 	// Pixels we treat in shader is pretty much hardcoded to 16 so no parameter has been used
 	const uint OutputTexturePixelsPerThread = 16;
 
-	// Scale U (x index) based on the number of pixels we treat per thread
-	const uint BaseU = ThreadId.x * OutputTexturePixelsPerThread;
-	const uint BaseV = ThreadId.y;
+	struct FOutputPixelData
+	{
+		uint2 UV;
+		float4 RGB;
+	};
+	FOutputPixelData PixelData[OutputTexturePixelsPerThread];
+
+	const uint SourcePixelStart = ThreadId.x * OutputTexturePixelsPerThread;
+	for (uint PixelIndex = 0; PixelIndex < OutputTexturePixelsPerThread; PixelIndex++)
+	{
+		PixelData[PixelIndex].UV.x = ((SourcePixelStart + PixelIndex) % OutputTextureSizeX); 
+		PixelData[PixelIndex].UV.y = ((SourcePixelStart + PixelIndex) / OutputTextureSizeX);
+		PixelData[PixelIndex].RGB = float4(0,0,0,1);
+	}
 
-	// Input buffer indexing is done in one dimension so adjust using thread Id and elements per row
-	const uint BufferIndex = ThreadId.y * HorizontalElementCount + ThreadId.x;
+	const uint BufferIndex = ThreadId.x;
 
 	// Used to normalize pixel values based on 10 bits input
 	const float Mul = 1.0 / 1023.0;
 
-	float3 RGB00 = uint3(0,0,0);
-	RGB00.x = ((((InputBuffer[BufferIndex].DWords[00] >>  0)  << 2) | ((InputBuffer[BufferIndex].DWords[00] >> 14) & 0x003)) & 0x3FF) * Mul;
-	RGB00.y = ((((InputBuffer[BufferIndex].DWords[00] >>  8)  << 4) | ((InputBuffer[BufferIndex].DWords[00] >> 20) & 0x00F)) & 0x3FF) * Mul;
-	RGB00.z = ((((InputBuffer[BufferIndex].DWords[00] >> 16)  << 6) | ((InputBuffer[BufferIndex].DWords[00] >> 26) & 0x03F)) & 0x3FF) * Mul;
+	PixelData[0].RGB.x = ((((InputBuffer[BufferIndex].DWords[00] >>  0)  << 2) | ((InputBuffer[BufferIndex].DWords[00] >> 14) & 0x003)) & 0x3FF) * Mul;
+	PixelData[0].RGB.y = ((((InputBuffer[BufferIndex].DWords[00] >>  8)  << 4) | ((InputBuffer[BufferIndex].DWords[00] >> 20) & 0x00F)) & 0x3FF) * Mul;
+	PixelData[0].RGB.z = ((((InputBuffer[BufferIndex].DWords[00] >> 16)  << 6) | ((InputBuffer[BufferIndex].DWords[00] >> 26) & 0x03F)) & 0x3FF) * Mul;
 
-	float3 RGB01 = uint3(0,0,0);
-	RGB01.x = ((((InputBuffer[BufferIndex].DWords[00] >> 24) << 8) | ((InputBuffer[BufferIndex].DWords[01] >> 00) & 0x0FF)) & 0x3FF) * Mul;
-	RGB01.y = ((((InputBuffer[BufferIndex].DWords[01] >>  8) << 2) | ((InputBuffer[BufferIndex].DWords[01] >> 22) & 0x003)) & 0x3FF) * Mul;
-	RGB01.z = ((((InputBuffer[BufferIndex].DWords[01] >> 16) << 4) | ((InputBuffer[BufferIndex].DWords[01] >> 28) & 0x00F)) & 0x3FF) * Mul;
+	PixelData[1].RGB.x = ((((InputBuffer[BufferIndex].DWords[00] >> 24) << 8) | ((InputBuffer[BufferIndex].DWords[01] >> 00) & 0x0FF)) & 0x3FF) * Mul;
+	PixelData[1].RGB.y = ((((InputBuffer[BufferIndex].DWords[01] >>  8) << 2) | ((InputBuffer[BufferIndex].DWords[01] >> 22) & 0x003)) & 0x3FF) * Mul;
+	PixelData[1].RGB.z = ((((InputBuffer[BufferIndex].DWords[01] >> 16) << 4) | ((InputBuffer[BufferIndex].DWords[01] >> 28) & 0x00F)) & 0x3FF) * Mul;
 
-	float3 RGB02 = uint3(0,0,0);
-	RGB02.x = ((((InputBuffer[BufferIndex].DWords[01] >> 24) << 6) | ((InputBuffer[BufferIndex].DWords[02] >> 02) & 0x03F)) & 0x3FF) * Mul;
-	RGB02.y = ((((InputBuffer[BufferIndex].DWords[02] >>  0) << 8) | ((InputBuffer[BufferIndex].DWords[02] >>  8) & 0x0FF)) & 0x3FF) * Mul;
-	RGB02.z = ((((InputBuffer[BufferIndex].DWords[02] >> 16) << 2) | ((InputBuffer[BufferIndex].DWords[02] >> 30) & 0x002)) & 0x3FF) * Mul;
+	PixelData[2].RGB.x = ((((InputBuffer[BufferIndex].DWords[01] >> 24) << 6) | ((InputBuffer[BufferIndex].DWords[02] >> 02) & 0x03F)) & 0x3FF) * Mul;
+	PixelData[2].RGB.y = ((((InputBuffer[BufferIndex].DWords[02] >>  0) << 8) | ((InputBuffer[BufferIndex].DWords[02] >>  8) & 0x0FF)) & 0x3FF) * Mul;
+	PixelData[2].RGB.z = ((((InputBuffer[BufferIndex].DWords[02] >> 16) << 2) | ((InputBuffer[BufferIndex].DWords[02] >> 30) & 0x002)) & 0x3FF) * Mul;
 	
-	float3 RGB03 = uint3(0,0,0);
-	RGB03.x = ((((InputBuffer[BufferIndex].DWords[02] >> 24) << 4) | ((InputBuffer[BufferIndex].DWords[03] >> 04) & 0x00F)) & 0x3FF) * Mul;
-	RGB03.y = ((((InputBuffer[BufferIndex].DWords[03] >> 00) << 6) | ((InputBuffer[BufferIndex].DWords[03] >> 10) & 0x03F)) & 0x3FF) * Mul;
-	RGB03.z = ((((InputBuffer[BufferIndex].DWords[03] >>  8) << 8) | ((InputBuffer[BufferIndex].DWords[03] >> 16) & 0x0FF)) & 0x3FF) * Mul;
+	PixelData[3].RGB.x = ((((InputBuffer[BufferIndex].DWords[02] >> 24) << 4) | ((InputBuffer[BufferIndex].DWords[03] >> 04) & 0x00F)) & 0x3FF) * Mul;
+	PixelData[3].RGB.y = ((((InputBuffer[BufferIndex].DWords[03] >> 00) << 6) | ((InputBuffer[BufferIndex].DWords[03] >> 10) & 0x03F)) & 0x3FF) * Mul;
+	PixelData[3].RGB.z = ((((InputBuffer[BufferIndex].DWords[03] >>  8) << 8) | ((InputBuffer[BufferIndex].DWords[03] >> 16) & 0x0FF)) & 0x3FF) * Mul;
 	
-	float3 RGB04 = uint3(0,0,0);
-	RGB04.x = ((((InputBuffer[BufferIndex].DWords[03] >> 24) << 2) | ((InputBuffer[BufferIndex].DWords[04] >> 06) & 0x003)) & 0x3FF) * Mul;
-	RGB04.y = ((((InputBuffer[BufferIndex].DWords[04] >> 00) << 4) | ((InputBuffer[BufferIndex].DWords[04] >> 12) & 0x00F)) & 0x3FF) * Mul;
-	RGB04.z = ((((InputBuffer[BufferIndex].DWords[04] >>  8) << 6) | ((InputBuffer[BufferIndex].DWords[04] >> 18) & 0x03F)) & 0x3FF) * Mul;
+	PixelData[4].RGB.x = ((((InputBuffer[BufferIndex].DWords[03] >> 24) << 2) | ((InputBuffer[BufferIndex].DWords[04] >> 06) & 0x003)) & 0x3FF) * Mul;
+	PixelData[4].RGB.y = ((((InputBuffer[BufferIndex].DWords[04] >> 00) << 4) | ((InputBuffer[BufferIndex].DWords[04] >> 12) & 0x00F)) & 0x3FF) * Mul;
+	PixelData[4].RGB.z = ((((InputBuffer[BufferIndex].DWords[04] >>  8) << 6) | ((InputBuffer[BufferIndex].DWords[04] >> 18) & 0x03F)) & 0x3FF) * Mul;
 	
-	float3 RGB05;
-	RGB05.x = ((((InputBuffer[BufferIndex].DWords[04] >> 16) << 8) | ((InputBuffer[BufferIndex].DWords[04] >> 24) & 0x0FF)) & 0x3FF) * Mul;
-	RGB05.y = ((((InputBuffer[BufferIndex].DWords[05] >> 00) << 2) | ((InputBuffer[BufferIndex].DWords[05] >> 14) & 0x003)) & 0x3FF) * Mul;
-	RGB05.z = ((((InputBuffer[BufferIndex].DWords[05] >>  8) << 4) | ((InputBuffer[BufferIndex].DWords[05] >> 20) & 0x00F)) & 0x3FF) * Mul;
-
-	float3 RGB06 = uint3(0,0,0);
-	RGB06.x = ((((InputBuffer[BufferIndex].DWords[05] >> 16) << 6) | ((InputBuffer[BufferIndex].DWords[05] >> 26) & 0x03F)) & 0x3FF) * Mul;
-	RGB06.y = ((((InputBuffer[BufferIndex].DWords[05] >> 24) << 8) | ((InputBuffer[BufferIndex].DWords[06] >> 00) & 0x0FF)) & 0x3FF) * Mul;
-	RGB06.z = ((((InputBuffer[BufferIndex].DWords[06] >>  8) << 2) | ((InputBuffer[BufferIndex].DWords[06] >> 22) & 0x003)) & 0x3FF) * Mul;
-
-	float3 RGB07 = uint3(0,0,0);
-	RGB07.x = ((((InputBuffer[BufferIndex].DWords[06] >> 16) << 4) | ((InputBuffer[BufferIndex].DWords[06] >> 28) & 0x00F)) & 0x3FF) * Mul;
-	RGB07.y = ((((InputBuffer[BufferIndex].DWords[06] >> 24) << 6) | ((InputBuffer[BufferIndex].DWords[07] >> 02) & 0x03F)) & 0x3FF) * Mul;
-	RGB07.z = ((((InputBuffer[BufferIndex].DWords[07] >>  0) << 8) | ((InputBuffer[BufferIndex].DWords[07] >>  8) & 0x0FF)) & 0x3FF) * Mul;
-
-	float3 RGB08 = uint3(0,0,0);
-	RGB08.x = ((((InputBuffer[BufferIndex].DWords[07] >> 16) << 2) | ((InputBuffer[BufferIndex].DWords[07] >> 30) & 0x003)) & 0x3FF) * Mul;
-	RGB08.y = ((((InputBuffer[BufferIndex].DWords[07] >> 24) << 4) | ((InputBuffer[BufferIndex].DWords[ 8] >> 04) & 0x00F)) & 0x3FF) * Mul;
-	RGB08.z = ((((InputBuffer[BufferIndex].DWords[ 8] >>  0) << 6) | ((InputBuffer[BufferIndex].DWords[ 8] >> 10) & 0x03F)) & 0x3FF) * Mul;
-
-	float3 RGB09 = uint3(0,0,0);
-	RGB09.x = ((((InputBuffer[BufferIndex].DWords[ 8] >>  8) << 8) | ((InputBuffer[BufferIndex].DWords[ 8] >> 16) & 0x0FF)) & 0x3FF) * Mul;
-	RGB09.y = ((((InputBuffer[BufferIndex].DWords[ 8] >> 24) << 2) | ((InputBuffer[BufferIndex].DWords[ 9] >> 06) & 0x003)) & 0x3FF) * Mul;
-	RGB09.z = ((((InputBuffer[BufferIndex].DWords[ 9] >>  0) << 4) | ((InputBuffer[BufferIndex].DWords[ 9] >> 12) & 0x00F)) & 0x3FF) * Mul;
-
-	float3 RGB10 = uint3(0,0,0);
-	RGB10.x = ((((InputBuffer[BufferIndex].DWords[ 9] >>  8) << 6) | ((InputBuffer[BufferIndex].DWords[ 9] >> 18) & 0x03F)) & 0x3FF) * Mul;
-	RGB10.y = ((((InputBuffer[BufferIndex].DWords[ 9] >> 16) << 8) | ((InputBuffer[BufferIndex].DWords[ 9] >> 24) & 0x0FF)) & 0x3FF) * Mul;
-	RGB10.z = ((((InputBuffer[BufferIndex].DWords[10] >>  0) << 2) | ((InputBuffer[BufferIndex].DWords[10] >> 14) & 0x003)) & 0x3FF) * Mul;
-
-	float3 RGB11 = uint3(0,0,0);
-	RGB11.x = ((((InputBuffer[BufferIndex].DWords[10] >>  8) << 4) | ((InputBuffer[BufferIndex].DWords[10] >> 20) & 0x00F)) & 0x3FF) * Mul;
-	RGB11.y = ((((InputBuffer[BufferIndex].DWords[10] >> 16) << 6) | ((InputBuffer[BufferIndex].DWords[10] >> 26) & 0x03F)) & 0x3FF) * Mul;
-	RGB11.z = ((((InputBuffer[BufferIndex].DWords[10] >>  24) << 8) | ((InputBuffer[BufferIndex].DWords[11] >> 00) & 0x0FF)) & 0x3FF) * Mul;
+	PixelData[5].RGB.x = ((((InputBuffer[BufferIndex].DWords[04] >> 16) << 8) | ((InputBuffer[BufferIndex].DWords[04] >> 24) & 0x0FF)) & 0x3FF) * Mul;
+	PixelData[5].RGB.y = ((((InputBuffer[BufferIndex].DWords[05] >> 00) << 2) | ((InputBuffer[BufferIndex].DWords[05] >> 14) & 0x003)) & 0x3FF) * Mul;
+	PixelData[5].RGB.z = ((((InputBuffer[BufferIndex].DWords[05] >>  8) << 4) | ((InputBuffer[BufferIndex].DWords[05] >> 20) & 0x00F)) & 0x3FF) * Mul;
+
+	PixelData[6].RGB.x = ((((InputBuffer[BufferIndex].DWords[05] >> 16) << 6) | ((InputBuffer[BufferIndex].DWords[05] >> 26) & 0x03F)) & 0x3FF) * Mul;
+	PixelData[6].RGB.y = ((((InputBuffer[BufferIndex].DWords[05] >> 24) << 8) | ((InputBuffer[BufferIndex].DWords[06] >> 00) & 0x0FF)) & 0x3FF) * Mul;
+	PixelData[6].RGB.z = ((((InputBuffer[BufferIndex].DWords[06] >>  8) << 2) | ((InputBuffer[BufferIndex].DWords[06] >> 22) & 0x003)) & 0x3FF) * Mul;
+
+	PixelData[7].RGB.x = ((((InputBuffer[BufferIndex].DWords[06] >> 16) << 4) | ((InputBuffer[BufferIndex].DWords[06] >> 28) & 0x00F)) & 0x3FF) * Mul;
+	PixelData[7].RGB.y = ((((InputBuffer[BufferIndex].DWords[06] >> 24) << 6) | ((InputBuffer[BufferIndex].DWords[07] >> 02) & 0x03F)) & 0x3FF) * Mul;
+	PixelData[7].RGB.z = ((((InputBuffer[BufferIndex].DWords[07] >>  0) << 8) | ((InputBuffer[BufferIndex].DWords[07] >>  8) & 0x0FF)) & 0x3FF) * Mul;
+
+	PixelData[8].RGB.x = ((((InputBuffer[BufferIndex].DWords[07] >> 16) << 2) | ((InputBuffer[BufferIndex].DWords[07] >> 30) & 0x003)) & 0x3FF) * Mul;
+	PixelData[8].RGB.y = ((((InputBuffer[BufferIndex].DWords[07] >> 24) << 4) | ((InputBuffer[BufferIndex].DWords[ 8] >> 04) & 0x00F)) & 0x3FF) * Mul;
+	PixelData[8].RGB.z = ((((InputBuffer[BufferIndex].DWords[ 8] >>  0) << 6) | ((InputBuffer[BufferIndex].DWords[ 8] >> 10) & 0x03F)) & 0x3FF) * Mul;
+
+	PixelData[9].RGB.x = ((((InputBuffer[BufferIndex].DWords[ 8] >>  8) << 8) | ((InputBuffer[BufferIndex].DWords[ 8] >> 16) & 0x0FF)) & 0x3FF) * Mul;
+	PixelData[9].RGB.y = ((((InputBuffer[BufferIndex].DWords[ 8] >> 24) << 2) | ((InputBuffer[BufferIndex].DWords[ 9] >> 06) & 0x003)) & 0x3FF) * Mul;
+	PixelData[9].RGB.z = ((((InputBuffer[BufferIndex].DWords[ 9] >>  0) << 4) | ((InputBuffer[BufferIndex].DWords[ 9] >> 12) & 0x00F)) & 0x3FF) * Mul;
+
+	PixelData[10].RGB.x = ((((InputBuffer[BufferIndex].DWords[ 9] >>  8) << 6) | ((InputBuffer[BufferIndex].DWords[ 9] >> 18) & 0x03F)) & 0x3FF) * Mul;
+	PixelData[10].RGB.y = ((((InputBuffer[BufferIndex].DWords[ 9] >> 16) << 8) | ((InputBuffer[BufferIndex].DWords[ 9] >> 24) & 0x0FF)) & 0x3FF) * Mul;
+	PixelData[10].RGB.z = ((((InputBuffer[BufferIndex].DWords[10] >>  0) << 2) | ((InputBuffer[BufferIndex].DWords[10] >> 14) & 0x003)) & 0x3FF) * Mul;
+
+	PixelData[11].RGB.x = ((((InputBuffer[BufferIndex].DWords[10] >>  8) << 4) | ((InputBuffer[BufferIndex].DWords[10] >> 20) & 0x00F)) & 0x3FF) * Mul;
+	PixelData[11].RGB.y = ((((InputBuffer[BufferIndex].DWords[10] >> 16) << 6) | ((InputBuffer[BufferIndex].DWords[10] >> 26) & 0x03F)) & 0x3FF) * Mul;
+	PixelData[11].RGB.z = ((((InputBuffer[BufferIndex].DWords[10] >>  24) << 8) | ((InputBuffer[BufferIndex].DWords[11] >> 00) & 0x0FF)) & 0x3FF) * Mul;
 	
-	//OutRGB10Buffer[Index].DWords[12] = ((RGB13.y & 0x003) << 30) | ((RGB13.z >> 4)   << 24)								| ((RGB13.y >> 2)   << 16)							  | ((RGB13.x & 0xFF) << 8)  | ((RGB12.z & 0x3F) << 2)  | (RGB13.x >> 8);
-	float3 RGB12 = uint3(0,0,0);
-	RGB12.x = ((((InputBuffer[BufferIndex].DWords[11] >>  8) << 2) | ((InputBuffer[BufferIndex].DWords[11] >> 22) & 0x003)) & 0x3FF) * Mul;
-	RGB12.y = ((((InputBuffer[BufferIndex].DWords[11] >> 16) << 4) | ((InputBuffer[BufferIndex].DWords[11] >> 28) & 0x00F)) & 0x3FF) * Mul;
-	RGB12.z = ((((InputBuffer[BufferIndex].DWords[11] >> 24) << 6) | ((InputBuffer[BufferIndex].DWords[12] >> 02) & 0x03F)) & 0x3FF) * Mul;
-
-	float3 RGB13 = uint3(0,0,0);
-	RGB13.x = ((((InputBuffer[BufferIndex].DWords[12] >>  0) << 8) | ((InputBuffer[BufferIndex].DWords[12] >>  8) & 0x0FF)) & 0x3FF) * Mul;
-	RGB13.y = ((((InputBuffer[BufferIndex].DWords[12] >> 16) << 2) | ((InputBuffer[BufferIndex].DWords[12] >> 30) & 0x003)) & 0x3FF) * Mul;
-	RGB13.z = ((((InputBuffer[BufferIndex].DWords[12] >> 24) << 4) | ((InputBuffer[BufferIndex].DWords[13] >> 04) & 0x00F)) & 0x3FF) * Mul;
-
-	float3 RGB14 = uint3(0,0,0);
-	RGB14.x = ((((InputBuffer[BufferIndex].DWords[13] >>  0) << 6) | ((InputBuffer[BufferIndex].DWords[13] >> 10) & 0x03F)) & 0x3FF) * Mul;
-	RGB14.y = ((((InputBuffer[BufferIndex].DWords[13] >>  8) << 8) | ((InputBuffer[BufferIndex].DWords[13] >> 16) & 0x0FF)) & 0x3FF) * Mul;
-	RGB14.z = ((((InputBuffer[BufferIndex].DWords[13] >> 24) << 2) | ((InputBuffer[BufferIndex].DWords[14] >> 06) & 0x003)) & 0x3FF) * Mul;
-
-	float3 RGB15 = uint3(0,0,0);
-	RGB15.x = ((((InputBuffer[BufferIndex].DWords[14] >>  0) << 4) | ((InputBuffer[BufferIndex].DWords[14] >> 12) & 0x00F)) & 0x3FF) * Mul;
-	RGB15.y = ((((InputBuffer[BufferIndex].DWords[14] >>  8) << 6) | ((InputBuffer[BufferIndex].DWords[14] >> 18) & 0x03F)) & 0x3FF) * Mul;
-	RGB15.z = ((((InputBuffer[BufferIndex].DWords[14] >> 16) << 8) | ((InputBuffer[BufferIndex].DWords[14] >> 24) & 0x0FF)) & 0x3FF) * Mul;
+	PixelData[12].RGB.x = ((((InputBuffer[BufferIndex].DWords[11] >>  8) << 2) | ((InputBuffer[BufferIndex].DWords[11] >> 22) & 0x003)) & 0x3FF) * Mul;
+	PixelData[12].RGB.y = ((((InputBuffer[BufferIndex].DWords[11] >> 16) << 4) | ((InputBuffer[BufferIndex].DWords[11] >> 28) & 0x00F)) & 0x3FF) * Mul;
+	PixelData[12].RGB.z = ((((InputBuffer[BufferIndex].DWords[11] >> 24) << 6) | ((InputBuffer[BufferIndex].DWords[12] >> 02) & 0x03F)) & 0x3FF) * Mul;
+
+	PixelData[13].RGB.x = ((((InputBuffer[BufferIndex].DWords[12] >>  0) << 8) | ((InputBuffer[BufferIndex].DWords[12] >>  8) & 0x0FF)) & 0x3FF) * Mul;
+	PixelData[13].RGB.y = ((((InputBuffer[BufferIndex].DWords[12] >> 16) << 2) | ((InputBuffer[BufferIndex].DWords[12] >> 30) & 0x003)) & 0x3FF) * Mul;
+	PixelData[13].RGB.z = ((((InputBuffer[BufferIndex].DWords[12] >> 24) << 4) | ((InputBuffer[BufferIndex].DWords[13] >> 04) & 0x00F)) & 0x3FF) * Mul;
 
+	PixelData[14].RGB.x = ((((InputBuffer[BufferIndex].DWords[13] >>  0) << 6) | ((InputBuffer[BufferIndex].DWords[13] >> 10) & 0x03F)) & 0x3FF) * Mul;
+	PixelData[14].RGB.y = ((((InputBuffer[BufferIndex].DWords[13] >>  8) << 8) | ((InputBuffer[BufferIndex].DWords[13] >> 16) & 0x0FF)) & 0x3FF) * Mul;
+	PixelData[14].RGB.z = ((((InputBuffer[BufferIndex].DWords[13] >> 24) << 2) | ((InputBuffer[BufferIndex].DWords[14] >> 06) & 0x003)) & 0x3FF) * Mul;
+
+	PixelData[15].RGB.x = ((((InputBuffer[BufferIndex].DWords[14] >>  0) << 4) | ((InputBuffer[BufferIndex].DWords[14] >> 12) & 0x00F)) & 0x3FF) * Mul;
+	PixelData[15].RGB.y = ((((InputBuffer[BufferIndex].DWords[14] >>  8) << 6) | ((InputBuffer[BufferIndex].DWords[14] >> 18) & 0x03F)) & 0x3FF) * Mul;
+	PixelData[15].RGB.z = ((((InputBuffer[BufferIndex].DWords[14] >> 16) << 8) | ((InputBuffer[BufferIndex].DWords[14] >> 24) & 0x0FF)) & 0x3FF) * Mul;
+
+	for (uint PixelIndex = 0; PixelIndex < OutputTexturePixelsPerThread; PixelIndex++)
+	{
 	#if DO_SRGB_TO_LINEAR
-		RGB00 = sRGBToLinear(RGB00);
-		RGB01 = sRGBToLinear(RGB01);
-		RGB02 = sRGBToLinear(RGB02);
-		RGB03 = sRGBToLinear(RGB03);
-		RGB04 = sRGBToLinear(RGB04);
-		RGB05 = sRGBToLinear(RGB05);
-		RGB06 = sRGBToLinear(RGB06);
-		RGB07 = sRGBToLinear(RGB07);
-		RGB08 = sRGBToLinear(RGB08);
-		RGB09 = sRGBToLinear(RGB09);
-		RGB10 = sRGBToLinear(RGB10);
-		RGB11 = sRGBToLinear(RGB11);
-		RGB12 = sRGBToLinear(RGB12);
-		RGB13 = sRGBToLinear(RGB13);
-		RGB14 = sRGBToLinear(RGB14);
-		RGB15 = sRGBToLinear(RGB15);
+		PixelData[PixelIndex].RGB.xyz = sRGBToLinear(PixelData[PixelIndex].RGB.xyz);
 	#endif
-
-	OutTexture[uint2(BaseU + 0, BaseV)] = float4(RGB00, 1.0);
-	OutTexture[uint2(BaseU + 1, BaseV)] = float4(RGB01, 1.0);
-	OutTexture[uint2(BaseU + 2, BaseV)] = float4(RGB02, 1.0);
-	OutTexture[uint2(BaseU + 3, BaseV)] = float4(RGB03, 1.0);
-	OutTexture[uint2(BaseU + 4, BaseV)] = float4(RGB04, 1.0);
-	OutTexture[uint2(BaseU + 5, BaseV)] = float4(RGB05, 1.0);
-	OutTexture[uint2(BaseU + 6, BaseV)] = float4(RGB06, 1.0);
-	OutTexture[uint2(BaseU + 7, BaseV)] = float4(RGB07, 1.0);
-	OutTexture[uint2(BaseU + 8, BaseV)] = float4(RGB08, 1.0);
-	OutTexture[uint2(BaseU + 9, BaseV)] = float4(RGB09, 1.0);
-	OutTexture[uint2(BaseU + 10, BaseV)] = float4(RGB10, 1.0);
-	OutTexture[uint2(BaseU + 11, BaseV)] = float4(RGB11, 1.0);
-	OutTexture[uint2(BaseU + 12, BaseV)] = float4(RGB12, 1.0);
-	OutTexture[uint2(BaseU + 13, BaseV)] = float4(RGB13, 1.0);
-	OutTexture[uint2(BaseU + 14, BaseV)] = float4(RGB14, 1.0);
-	OutTexture[uint2(BaseU + 15, BaseV)] = float4(RGB15, 1.0);
+		
+		OutTexture[PixelData[PixelIndex].UV] = PixelData[PixelIndex].RGB;
+	}
 }
 
 // Shader to convert from RGBA to RGB8
 
-[numthreads(8,8,1)]
+[numthreads(64,1,1)]
 void RGBToRGB8BitCS(uint3 ThreadId : SV_DispatchThreadID, uint GroupIndex : SV_GroupIndex)
 {
 	// Make sure it's a valid thread
-	if(ThreadId.x >= HorizontalElementCount || ThreadId.y >= VerticalElementCount)
+	if(ThreadId.x >= TotalElementCount)
 	{
 		return;
 	}
 
-	uint Index = ThreadId.y * HorizontalElementCount + ThreadId.x;
-
-	float X = (ThreadId.x * InputTexturePixelsPerThread) + InputPixelOffsetX;
-	float Y = (ThreadId.y + InputPixelOffsetY) * OnePixelDeltaY;
-
-	// We process 4 RGB pixels per thread
-	uint Denormalizer = 255;
-	uint3 RGB0 = InputTexture.SampleLevel(InputSampler, float2(((X + 0.0) + 0.5) * OnePixelDeltaX, Y), 0).rgb * Denormalizer;
-	uint3 RGB1 = InputTexture.SampleLevel(InputSampler, float2(((X + 1.0) + 0.5) * OnePixelDeltaX, Y), 0).rgb * Denormalizer;
-	uint3 RGB2 = InputTexture.SampleLevel(InputSampler, float2(((X + 2.0) + 0.5) * OnePixelDeltaX, Y), 0).rgb * Denormalizer;
-	uint3 RGB3 = InputTexture.SampleLevel(InputSampler, float2(((X + 3.0) + 0.5) * OnePixelDeltaX, Y), 0).rgb * Denormalizer;
-
-	// Cap values to be in range
-	RGB0 = min(RGB0, Denormalizer);
-	RGB1 = min(RGB1, Denormalizer);
-	RGB2 = min(RGB2, Denormalizer);
-	RGB3 = min(RGB3, Denormalizer);
+	// Convert output buffer index to pixels
+	const uint PixelCount = 4;
+	const uint SourcePixelStart = (ThreadId.x * PixelCount);
+	const uint Denormalizer = 255;
+	const uint Index = ThreadId.x;
+	
+	uint3 RGB[PixelCount];
+	for (uint PixelIndex = 0; PixelIndex < PixelCount; PixelIndex++)
+	{
+		const float BaseX = ((SourcePixelStart + PixelIndex) % CapturedSizeX) + InputPixelOffsetX; 
+		const float BaseY = ((SourcePixelStart + PixelIndex) / CapturedSizeX) + InputPixelOffsetY;
+		const float U = (BaseX + 0.5) * OnePixelDeltaX;
+		const float V = (BaseY + 0.5) * OnePixelDeltaY;
+		RGB[PixelIndex] = InputTexture.SampleLevel(InputSampler, float2(U,V), 0).rgb * Denormalizer; 
+		RGB[PixelIndex] = min(RGB[PixelIndex], Denormalizer);
+	}
 
 	// We write the 4 RGB pixels (24 bits * 4 = 96 bits = 3x32 bits) following the packing 2110 expects
-	OutRGB8Buffer[Index].DWord0 = (RGB1.x << 24) | (RGB0.z << 16) | (RGB0.y << 8) | (RGB0.x << 0);
-	OutRGB8Buffer[Index].DWord1 = (RGB2.y << 24) | (RGB2.x << 16) | (RGB1.z << 8) | (RGB1.y << 0);
-	OutRGB8Buffer[Index].DWord2 = (RGB3.z << 24) | (RGB3.y << 16) | (RGB3.x << 8) | (RGB2.z << 0);
+	OutRGB8Buffer[Index].DWord0 = (RGB[1].x << 24) | (RGB[0].z << 16) | (RGB[0].y << 8) | (RGB[0].x << 0);
+	OutRGB8Buffer[Index].DWord1 = (RGB[2].y << 24) | (RGB[2].x << 16) | (RGB[1].z << 8) | (RGB[1].y << 0);
+	OutRGB8Buffer[Index].DWord2 = (RGB[3].z << 24) | (RGB[3].y << 16) | (RGB[3].x << 8) | (RGB[2].z << 0);
 }
 
 // Shader to convert from packed RGB8 to RGBA8 texture
 
-[numthreads(8,8,1)]
+[numthreads(64,1,1)]
 void RGB8BitToRGBA8CS(uint3 ThreadId : SV_DispatchThreadID, uint GroupIndex : SV_GroupIndex)
 {
 	// Make sure it's a valid thread
-	if(ThreadId.x >= HorizontalElementCount || ThreadId.y >= VerticalElementCount)
+	if(ThreadId.x >= TotalElementCount)
 	{
 		return;
 	}
 
-	// Pixels we treat in shader is pretty much hardcoded to 16 so no parameter has been used
+	// Pixels we treat in shader is pretty much hardcoded to 4 so no parameter has been used
 	const uint OutputTexturePixelsPerThread = 4;
 
-	// Scale U (x index) based on the number of pixels we treat per thread
-	const uint BaseU = ThreadId.x * OutputTexturePixelsPerThread;
-	const uint BaseV = ThreadId.y;
+	struct FOutputPixelData
+	{
+		uint2 UV;
+		float4 RGB;
+	};
+	FOutputPixelData PixelData[OutputTexturePixelsPerThread];
+
+	const uint SourcePixelStart = ThreadId.x * OutputTexturePixelsPerThread;
+	for (uint PixelIndex = 0; PixelIndex < OutputTexturePixelsPerThread; PixelIndex++)
+	{
+		PixelData[PixelIndex].UV.x = ((SourcePixelStart + PixelIndex) % OutputTextureSizeX); 
+		PixelData[PixelIndex].UV.y = ((SourcePixelStart + PixelIndex) / OutputTextureSizeX);
+		PixelData[PixelIndex].RGB = float4(0,0,0,1);
+	}
 
-	// Input buffer indexing is done in one dimension so adjust using thread Id and elements per row
-	const uint BufferIndex = ThreadId.y * HorizontalElementCount + ThreadId.x;
+	const uint BufferIndex = ThreadId.x;
 
 	// Used to normalize pixel values based on 8 bits input
 	const float Mul = 1.0 / 255.0;
 
-	float3 RGB00;
-	RGB00.x = (((InputRGB8Buffer[BufferIndex].DWord0 >> 00) & 0xFF)) * Mul;
-	RGB00.y = (((InputRGB8Buffer[BufferIndex].DWord0 >>  8) & 0xFF)) * Mul;
-	RGB00.z = (((InputRGB8Buffer[BufferIndex].DWord0 >> 16) & 0xFF)) * Mul;
+	PixelData[0].RGB.x = (((InputRGB8Buffer[BufferIndex].DWord0 >> 00) & 0xFF)) * Mul;
+	PixelData[0].RGB.y = (((InputRGB8Buffer[BufferIndex].DWord0 >>  8) & 0xFF)) * Mul;
+	PixelData[0].RGB.z = (((InputRGB8Buffer[BufferIndex].DWord0 >> 16) & 0xFF)) * Mul;
 
-	float3 RGB01;
-	RGB01.x = (((InputRGB8Buffer[BufferIndex].DWord0 >> 24) & 0xFF)) * Mul;
-	RGB01.y = (((InputRGB8Buffer[BufferIndex].DWord1 >> 00) & 0xFF)) * Mul;
-	RGB01.z = (((InputRGB8Buffer[BufferIndex].DWord1 >>  8) & 0xFF)) * Mul;
+	PixelData[1].RGB.x = (((InputRGB8Buffer[BufferIndex].DWord0 >> 24) & 0xFF)) * Mul;
+	PixelData[1].RGB.y = (((InputRGB8Buffer[BufferIndex].DWord1 >> 00) & 0xFF)) * Mul;
+	PixelData[1].RGB.z = (((InputRGB8Buffer[BufferIndex].DWord1 >>  8) & 0xFF)) * Mul;
 
-	float3 RGB02;
-	RGB02.x = (((InputRGB8Buffer[BufferIndex].DWord1 >> 16) & 0xFF)) * Mul;
-	RGB02.y = (((InputRGB8Buffer[BufferIndex].DWord1 >> 24) & 0xFF)) * Mul;
-	RGB02.z = (((InputRGB8Buffer[BufferIndex].DWord2 >> 00) & 0xFF)) * Mul;
+	PixelData[2].RGB.x = (((InputRGB8Buffer[BufferIndex].DWord1 >> 16) & 0xFF)) * Mul;
+	PixelData[2].RGB.y = (((InputRGB8Buffer[BufferIndex].DWord1 >> 24) & 0xFF)) * Mul;
+	PixelData[2].RGB.z = (((InputRGB8Buffer[BufferIndex].DWord2 >> 00) & 0xFF)) * Mul;
 	
-	float3 RGB03;
-	RGB03.x = (((InputRGB8Buffer[BufferIndex].DWord2 >>  8) & 0xFF)) * Mul;
-	RGB03.y = (((InputRGB8Buffer[BufferIndex].DWord2 >> 16) & 0xFF)) * Mul;
-	RGB03.z = (((InputRGB8Buffer[BufferIndex].DWord2 >> 24) & 0xFF)) * Mul;
+	PixelData[3].RGB.x = (((InputRGB8Buffer[BufferIndex].DWord2 >>  8) & 0xFF)) * Mul;
+	PixelData[3].RGB.y = (((InputRGB8Buffer[BufferIndex].DWord2 >> 16) & 0xFF)) * Mul;
+	PixelData[3].RGB.z = (((InputRGB8Buffer[BufferIndex].DWord2 >> 24) & 0xFF)) * Mul;
 	
 
+	for (uint PixelIndex = 0; PixelIndex < OutputTexturePixelsPerThread; PixelIndex++)
+	{
 	#if DO_SRGB_TO_LINEAR
-		RGB00 = sRGBToLinear(RGB00);
-		RGB01 = sRGBToLinear(RGB01);
-		RGB02 = sRGBToLinear(RGB02);
-		RGB03 = sRGBToLinear(RGB03);
+		PixelData[PixelIndex].RGB.xyz = sRGBToLinear(PixelData[PixelIndex].RGB.xyz);
 	#endif
-
-	OutTexture[uint2(BaseU + 0, BaseV)] = float4(RGB00, 1.0);
-	OutTexture[uint2(BaseU + 1, BaseV)] = float4(RGB01, 1.0);
-	OutTexture[uint2(BaseU + 2, BaseV)] = float4(RGB02, 1.0);
-	OutTexture[uint2(BaseU + 3, BaseV)] = float4(RGB03, 1.0);
+		
+		OutTexture[PixelData[PixelIndex].UV] = PixelData[PixelIndex].RGB;
+	}
 }
 
 // Shader to convert from RGBA to RGB12
 
-[numthreads(8,8,1)]
+[numthreads(64,1,1)]
 void RGBToRGB12BitCS(uint3 ThreadId : SV_DispatchThreadID, uint GroupIndex : SV_GroupIndex)
 {
 	// Make sure it's a valid thread
-	if(ThreadId.x >= HorizontalElementCount || ThreadId.y >= VerticalElementCount)
+	if(ThreadId.x >= TotalElementCount)
 	{
 		return;
 	}
 
-	uint Index = ThreadId.y * HorizontalElementCount + ThreadId.x;
-
-	float X = (ThreadId.x * InputTexturePixelsPerThread) + InputPixelOffsetX;
-	float Y = (ThreadId.y + InputPixelOffsetY) * OnePixelDeltaY;
-
-	//12 bits max value
-	uint Denormalizer = 4095;
-	uint3 RGB00 = InputTexture.SampleLevel(InputSampler, float2(((X + 0.0) + 0.5) * OnePixelDeltaX, Y), 0).rgb * Denormalizer;
-	uint3 RGB01 = InputTexture.SampleLevel(InputSampler, float2(((X + 1.0) + 0.5) * OnePixelDeltaX, Y), 0).rgb * Denormalizer;
-	uint3 RGB02 = InputTexture.SampleLevel(InputSampler, float2(((X + 2.0) + 0.5) * OnePixelDeltaX, Y), 0).rgb * Denormalizer;
-	uint3 RGB03 = InputTexture.SampleLevel(InputSampler, float2(((X + 3.0) + 0.5) * OnePixelDeltaX, Y), 0).rgb * Denormalizer;
-	uint3 RGB04 = InputTexture.SampleLevel(InputSampler, float2(((X + 4.0) + 0.5) * OnePixelDeltaX, Y), 0).rgb * Denormalizer;
-	uint3 RGB05 = InputTexture.SampleLevel(InputSampler, float2(((X + 5.0) + 0.5) * OnePixelDeltaX, Y), 0).rgb * Denormalizer;
-	uint3 RGB06 = InputTexture.SampleLevel(InputSampler, float2(((X + 6.0) + 0.5) * OnePixelDeltaX, Y), 0).rgb * Denormalizer;
-	uint3 RGB07 = InputTexture.SampleLevel(InputSampler, float2(((X + 7.0) + 0.5) * OnePixelDeltaX, Y), 0).rgb * Denormalizer;
-
-	// Cap values to be in range
-	RGB00 = min(RGB00, Denormalizer);
-	RGB01 = min(RGB01, Denormalizer);
-	RGB02 = min(RGB02, Denormalizer);
-	RGB03 = min(RGB03, Denormalizer);
-	RGB04 = min(RGB04, Denormalizer);
-	RGB05 = min(RGB05, Denormalizer);
-	RGB06 = min(RGB06, Denormalizer);
-	RGB07 = min(RGB07, Denormalizer);
-
-	OutRGB12Buffer[Index].DWords[0] =							  ((RGB00.z >> 4)   << 24) |							 ((RGB00.y & 0xFF)  << 16) |((RGB00.x & 0xF)  << 12) | ((RGB00.y >> 8)    << 8)								| (RGB00.x >> 4);
-	OutRGB12Buffer[Index].DWords[1] = ((RGB01.y & 0x00F) << 28) | ((RGB01.z >> 8)   << 24) |							 ((RGB01.y >> 4)    << 16)							 | ((RGB01.x & 0xFF)  << 8) | ((RGB00.z & 0x00F) << 4)	| (RGB01.x >> 8);
-	OutRGB12Buffer[Index].DWords[2] =							  ((RGB02.y & 0xFF) << 24) | ((RGB02.x & 0xF) << 20)	| ((RGB02.y >> 8)   << 16)							  | ((RGB02.x >> 4)	  << 8)								| (RGB01.z & 0xFF);
-	OutRGB12Buffer[Index].DWords[3] =							  ((RGB03.y >> 4)   << 24)								| ((RGB03.x & 0xFF) << 16) | ((RGB02.z & 0xF) << 12) | ((RGB03.x >> 8)   << 8)								| (RGB02.z >> 4);
-	OutRGB12Buffer[Index].DWords[4] = ((RGB04.x & 0x00F) << 28) | ((RGB04.y >> 8)   << 24)								| ((RGB04.x >> 4)   << 16)							  | ((RGB03.z & 0xFF) << 8)  | ((RGB03.y & 0x00F)  << 4)| (RGB03.z >> 8);
-	OutRGB12Buffer[Index].DWords[5] =							  ((RGB05.x & 0xFF) << 24) | ((RGB04.z & 0xF) << 20)	| ((RGB05.x >> 8)   << 16)							  | ((RGB04.z >> 4)   << 8)								| (RGB04.y & 0xFF);
-	OutRGB12Buffer[Index].DWords[6] =							  ((RGB06.x >> 4)   << 24)								| ((RGB05.z & 0xFF) << 16) | ((RGB05.y & 0xF) << 12)  | ((RGB05.z >> 8)   << 8)								| (RGB05.y >> 4);
-	OutRGB12Buffer[Index].DWords[7] = ((RGB06.z & 0x00F) << 28) | ((RGB07.x >> 8)   << 24)								| ((RGB06.z >> 4)   << 16)							  | ((RGB06.y & 0xFF) << 8)  | ((RGB06.x & 0x00F) << 4) | (RGB06.y >> 8);
-	OutRGB12Buffer[Index].DWords[8] =							  ((RGB07.z & 0xFF) << 24) | ((RGB07.y & 0xF) << 20)	| ((RGB07.z >> 8)   << 16)							  | ((RGB07.y >> 4)   << 8)								| (RGB07.x & 0xFF);
+	// Convert output buffer index to pixels
+	const uint PixelCount = 8;
+	const uint SourcePixelStart = (ThreadId.x * PixelCount);
+	const uint Denormalizer = 4095;
+	const uint Index = ThreadId.x;
+	
+	uint3 RGB[PixelCount];
+	for (uint PixelIndex = 0; PixelIndex < PixelCount; PixelIndex++)
+	{
+		const float BaseX = ((SourcePixelStart + PixelIndex) % CapturedSizeX) + InputPixelOffsetX; 
+		const float BaseY = ((SourcePixelStart + PixelIndex) / CapturedSizeX) + InputPixelOffsetY;
+		const float U = (BaseX + 0.5) * OnePixelDeltaX;
+		const float V = (BaseY + 0.5) * OnePixelDeltaY;
+		RGB[PixelIndex] = InputTexture.SampleLevel(InputSampler, float2(U,V), 0).rgb * Denormalizer; 
+		RGB[PixelIndex] = min(RGB[PixelIndex], Denormalizer);
+	}
+
+	OutRGB12Buffer[Index].DWords[0] =							  ((RGB[0].z >> 4)   << 24) |							 ((RGB[0].y & 0xFF)  << 16) |((RGB[0].x & 0xF)  << 12) | ((RGB[0].y >> 8)    << 8)								| (RGB[0].x >> 4);
+	OutRGB12Buffer[Index].DWords[1] = ((RGB[1].y & 0x00F) << 28) | ((RGB[1].z >> 8)   << 24) |							 ((RGB[1].y >> 4)    << 16)							 | ((RGB[1].x & 0xFF)  << 8) | ((RGB[0].z & 0x00F) << 4)	| (RGB[1].x >> 8);
+	OutRGB12Buffer[Index].DWords[2] =							  ((RGB[2].y & 0xFF) << 24) | ((RGB[2].x & 0xF) << 20)	| ((RGB[2].y >> 8)   << 16)							  | ((RGB[2].x >> 4)	  << 8)								| (RGB[1].z & 0xFF);
+	OutRGB12Buffer[Index].DWords[3] =							  ((RGB[3].y >> 4)   << 24)								| ((RGB[3].x & 0xFF) << 16) | ((RGB[2].z & 0xF) << 12) | ((RGB[3].x >> 8)   << 8)								| (RGB[2].z >> 4);
+	OutRGB12Buffer[Index].DWords[4] = ((RGB[4].x & 0x00F) << 28) | ((RGB[4].y >> 8)   << 24)								| ((RGB[4].x >> 4)   << 16)							  | ((RGB[3].z & 0xFF) << 8)  | ((RGB[3].y & 0x00F)  << 4)| (RGB[3].z >> 8);
+	OutRGB12Buffer[Index].DWords[5] =							  ((RGB[5].x & 0xFF) << 24) | ((RGB[4].z & 0xF) << 20)	| ((RGB[5].x >> 8)   << 16)							  | ((RGB[4].z >> 4)   << 8)								| (RGB[4].y & 0xFF);
+	OutRGB12Buffer[Index].DWords[6] =							  ((RGB[6].x >> 4)   << 24)								| ((RGB[5].z & 0xFF) << 16) | ((RGB[5].y & 0xF) << 12)  | ((RGB[5].z >> 8)   << 8)								| (RGB[5].y >> 4);
+	OutRGB12Buffer[Index].DWords[7] = ((RGB[6].z & 0x00F) << 28) | ((RGB[7].x >> 8)   << 24)								| ((RGB[6].z >> 4)   << 16)							  | ((RGB[6].y & 0xFF) << 8)  | ((RGB[6].x & 0x00F) << 4) | (RGB[6].y >> 8);
+	OutRGB12Buffer[Index].DWords[8] =							  ((RGB[7].z & 0xFF) << 24) | ((RGB[7].y & 0xF) << 20)	| ((RGB[7].z >> 8)   << 16)							  | ((RGB[7].y >> 4)   << 8)								| (RGB[7].x & 0xFF);
 }
 
 // Shader to convert from packed RGB12 to float RGBA texture
 
-[numthreads(8,8,1)]
+[numthreads(64,1,1)]
 void RGB12BitToRGBACS(uint3 ThreadId : SV_DispatchThreadID, uint GroupIndex : SV_GroupIndex)
 {
 	// Make sure it's a valid thread
-	if(ThreadId.x >= HorizontalElementCount || ThreadId.y >= VerticalElementCount)
+	if(ThreadId.x >= TotalElementCount)
 	{
 		return;
 	}
 
-	// Pixels we treat in shader is pretty much hardcoded to 16 so no parameter has been used
+	// Pixels we treat in shader is pretty much hardcoded to 8 so no parameter has been used
 	const uint OutputTexturePixelsPerThread = 8;
 
-	// Scale U (x index) based on the number of pixels we treat per thread
-	const uint BaseU = ThreadId.x * OutputTexturePixelsPerThread;
-	const uint BaseV = ThreadId.y;
+	struct FOutputPixelData
+	{
+		uint2 UV;
+		float4 RGB;
+	};
+	FOutputPixelData PixelData[OutputTexturePixelsPerThread];
 
-	// Input buffer indexing is done in one dimension so adjust using thread Id and elements per row
-	const uint BufferIndex = ThreadId.y * HorizontalElementCount + ThreadId.x;
+	const uint SourcePixelStart = ThreadId.x * OutputTexturePixelsPerThread;
+	for (uint PixelIndex = 0; PixelIndex < OutputTexturePixelsPerThread; PixelIndex++)
+	{
+		PixelData[PixelIndex].UV.x = ((SourcePixelStart + PixelIndex) % OutputTextureSizeX); 
+		PixelData[PixelIndex].UV.y = ((SourcePixelStart + PixelIndex) / OutputTextureSizeX);
+		PixelData[PixelIndex].RGB = float4(0,0,0,1);
+	}
+
+	const uint BufferIndex = ThreadId.x;
 
 	// Used to normalize pixel values based on 12 bits input
 	const float Mul = 1.0 / 4095.0;
 
-	float3 RGB00 = uint3(0,0,0);
-	RGB00.x = ((((InputRGB12Buffer[BufferIndex].DWords[00] >>  0)  << 4) | ((InputRGB12Buffer[BufferIndex].DWords[00] >> 12) & 0x00F)) & 0xFFF) * Mul;
-	RGB00.y = ((((InputRGB12Buffer[BufferIndex].DWords[00] >>  8)  << 8) | ((InputRGB12Buffer[BufferIndex].DWords[00] >> 16) & 0x0FF)) & 0xFFF) * Mul;
-	RGB00.z = ((((InputRGB12Buffer[BufferIndex].DWords[00] >> 24)  << 4) | ((InputRGB12Buffer[BufferIndex].DWords[01] >> 04) & 0x00F)) & 0xFFF) * Mul;
+	PixelData[0].RGB.x = ((((InputRGB12Buffer[BufferIndex].DWords[00] >>  0)  << 4) | ((InputRGB12Buffer[BufferIndex].DWords[00] >> 12) & 0x00F)) & 0xFFF) * Mul;
+	PixelData[0].RGB.y = ((((InputRGB12Buffer[BufferIndex].DWords[00] >>  8)  << 8) | ((InputRGB12Buffer[BufferIndex].DWords[00] >> 16) & 0x0FF)) & 0xFFF) * Mul;
+	PixelData[0].RGB.z = ((((InputRGB12Buffer[BufferIndex].DWords[00] >> 24)  << 4) | ((InputRGB12Buffer[BufferIndex].DWords[01] >> 04) & 0x00F)) & 0xFFF) * Mul;
 
-	float3 RGB01 = uint3(0,0,0);
-	RGB01.x = ((((InputRGB12Buffer[BufferIndex].DWords[01] >>  0) << 8) | ((InputRGB12Buffer[BufferIndex].DWords[01] >>  8) & 0x0FF)) & 0xFFF) * Mul;
-	RGB01.y = ((((InputRGB12Buffer[BufferIndex].DWords[01] >> 16) << 4) | ((InputRGB12Buffer[BufferIndex].DWords[01] >> 28) & 0x00F)) & 0xFFF) * Mul;
-	RGB01.z = ((((InputRGB12Buffer[BufferIndex].DWords[01] >> 24) << 8) | ((InputRGB12Buffer[BufferIndex].DWords[02] >> 00) & 0x0FF)) & 0xFFF) * Mul;
+	PixelData[1].RGB.x = ((((InputRGB12Buffer[BufferIndex].DWords[01] >>  0) << 8) | ((InputRGB12Buffer[BufferIndex].DWords[01] >>  8) & 0x0FF)) & 0xFFF) * Mul;
+	PixelData[1].RGB.y = ((((InputRGB12Buffer[BufferIndex].DWords[01] >> 16) << 4) | ((InputRGB12Buffer[BufferIndex].DWords[01] >> 28) & 0x00F)) & 0xFFF) * Mul;
+	PixelData[1].RGB.z = ((((InputRGB12Buffer[BufferIndex].DWords[01] >> 24) << 8) | ((InputRGB12Buffer[BufferIndex].DWords[02] >> 00) & 0x0FF)) & 0xFFF) * Mul;
 
-	float3 RGB02 = uint3(0,0,0);
-	RGB02.x = ((((InputRGB12Buffer[BufferIndex].DWords[02] >>  8) << 4) | ((InputRGB12Buffer[BufferIndex].DWords[02] >> 20) & 0x00F)) & 0xFFF) * Mul;
-	RGB02.y = ((((InputRGB12Buffer[BufferIndex].DWords[02] >> 16) << 8) | ((InputRGB12Buffer[BufferIndex].DWords[02] >> 24) & 0x0FF)) & 0xFFF) * Mul;
-	RGB02.z = ((((InputRGB12Buffer[BufferIndex].DWords[03] >>  0) << 4) | ((InputRGB12Buffer[BufferIndex].DWords[03] >> 12) & 0x00F)) & 0xFFF) * Mul;
+	PixelData[2].RGB.x = ((((InputRGB12Buffer[BufferIndex].DWords[02] >>  8) << 4) | ((InputRGB12Buffer[BufferIndex].DWords[02] >> 20) & 0x00F)) & 0xFFF) * Mul;
+	PixelData[2].RGB.y = ((((InputRGB12Buffer[BufferIndex].DWords[02] >> 16) << 8) | ((InputRGB12Buffer[BufferIndex].DWords[02] >> 24) & 0x0FF)) & 0xFFF) * Mul;
+	PixelData[2].RGB.z = ((((InputRGB12Buffer[BufferIndex].DWords[03] >>  0) << 4) | ((InputRGB12Buffer[BufferIndex].DWords[03] >> 12) & 0x00F)) & 0xFFF) * Mul;
 	
-	float3 RGB03 = uint3(0,0,0);
-	RGB03.x = ((((InputRGB12Buffer[BufferIndex].DWords[03] >>  8) << 8) | ((InputRGB12Buffer[BufferIndex].DWords[03] >> 16) & 0x0FF)) & 0xFFF) * Mul;
-	RGB03.y = ((((InputRGB12Buffer[BufferIndex].DWords[03] >> 24) << 4) | ((InputRGB12Buffer[BufferIndex].DWords[04] >> 04) & 0x00F)) & 0xFFF) * Mul;
-	RGB03.z = ((((InputRGB12Buffer[BufferIndex].DWords[04] >>  0) << 8) | ((InputRGB12Buffer[BufferIndex].DWords[04] >>  8) & 0x0FF)) & 0xFFF) * Mul;
+	PixelData[3].RGB.x = ((((InputRGB12Buffer[BufferIndex].DWords[03] >>  8) << 8) | ((InputRGB12Buffer[BufferIndex].DWords[03] >> 16) & 0x0FF)) & 0xFFF) * Mul;
+	PixelData[3].RGB.y = ((((InputRGB12Buffer[BufferIndex].DWords[03] >> 24) << 4) | ((InputRGB12Buffer[BufferIndex].DWords[04] >> 04) & 0x00F)) & 0xFFF) * Mul;
+	PixelData[3].RGB.z = ((((InputRGB12Buffer[BufferIndex].DWords[04] >>  0) << 8) | ((InputRGB12Buffer[BufferIndex].DWords[04] >>  8) & 0x0FF)) & 0xFFF) * Mul;
 	
-	float3 RGB04 = uint3(0,0,0);
-	RGB04.x = ((((InputRGB12Buffer[BufferIndex].DWords[04] >> 16) << 4) | ((InputRGB12Buffer[BufferIndex].DWords[04] >> 28) & 0x00F)) & 0xFFF) * Mul;
-	RGB04.y = ((((InputRGB12Buffer[BufferIndex].DWords[04] >> 24) << 8) | ((InputRGB12Buffer[BufferIndex].DWords[05] >>  0) & 0x0FF)) & 0xFFF) * Mul;
-	RGB04.z = ((((InputRGB12Buffer[BufferIndex].DWords[05] >>  8) << 4) | ((InputRGB12Buffer[BufferIndex].DWords[05] >> 20) & 0x00F)) & 0xFFF) * Mul;
+	PixelData[4].RGB.x = ((((InputRGB12Buffer[BufferIndex].DWords[04] >> 16) << 4) | ((InputRGB12Buffer[BufferIndex].DWords[04] >> 28) & 0x00F)) & 0xFFF) * Mul;
+	PixelData[4].RGB.y = ((((InputRGB12Buffer[BufferIndex].DWords[04] >> 24) << 8) | ((InputRGB12Buffer[BufferIndex].DWords[05] >>  0) & 0x0FF)) & 0xFFF) * Mul;
+	PixelData[4].RGB.z = ((((InputRGB12Buffer[BufferIndex].DWords[05] >>  8) << 4) | ((InputRGB12Buffer[BufferIndex].DWords[05] >> 20) & 0x00F)) & 0xFFF) * Mul;
 	
-	float3 RGB05;
-	RGB05.x = ((((InputRGB12Buffer[BufferIndex].DWords[05] >> 16) << 8) | ((InputRGB12Buffer[BufferIndex].DWords[05] >> 24) & 0x0FF)) & 0xFFF) * Mul;
-	RGB05.y = ((((InputRGB12Buffer[BufferIndex].DWords[06] >> 00) << 4) | ((InputRGB12Buffer[BufferIndex].DWords[06] >> 12) & 0x00F)) & 0xFFF) * Mul;
-	RGB05.z = ((((InputRGB12Buffer[BufferIndex].DWords[06] >>  8) << 8) | ((InputRGB12Buffer[BufferIndex].DWords[06] >> 16) & 0x0FF)) & 0xFFF) * Mul;
+	PixelData[5].RGB.x = ((((InputRGB12Buffer[BufferIndex].DWords[05] >> 16) << 8) | ((InputRGB12Buffer[BufferIndex].DWords[05] >> 24) & 0x0FF)) & 0xFFF) * Mul;
+	PixelData[5].RGB.y = ((((InputRGB12Buffer[BufferIndex].DWords[06] >> 00) << 4) | ((InputRGB12Buffer[BufferIndex].DWords[06] >> 12) & 0x00F)) & 0xFFF) * Mul;
+	PixelData[5].RGB.z = ((((InputRGB12Buffer[BufferIndex].DWords[06] >>  8) << 8) | ((InputRGB12Buffer[BufferIndex].DWords[06] >> 16) & 0x0FF)) & 0xFFF) * Mul;
 
-	float3 RGB06 = uint3(0,0,0);
-	RGB06.x = ((((InputRGB12Buffer[BufferIndex].DWords[06] >> 24) << 4) | ((InputRGB12Buffer[BufferIndex].DWords[07] >> 04) & 0x00F)) & 0xFFF) * Mul;
-	RGB06.y = ((((InputRGB12Buffer[BufferIndex].DWords[07] >>  0) << 8) | ((InputRGB12Buffer[BufferIndex].DWords[07] >>  8) & 0x0FF)) & 0xFFF) * Mul;
-	RGB06.z = ((((InputRGB12Buffer[BufferIndex].DWords[07] >>  16) << 4) | ((InputRGB12Buffer[BufferIndex].DWords[07] >> 28) & 0x00F)) & 0xFFF) * Mul;
+	PixelData[6].RGB.x = ((((InputRGB12Buffer[BufferIndex].DWords[06] >> 24) << 4) | ((InputRGB12Buffer[BufferIndex].DWords[07] >> 04) & 0x00F)) & 0xFFF) * Mul;
+	PixelData[6].RGB.y = ((((InputRGB12Buffer[BufferIndex].DWords[07] >>  0) << 8) | ((InputRGB12Buffer[BufferIndex].DWords[07] >>  8) & 0x0FF)) & 0xFFF) * Mul;
+	PixelData[6].RGB.z = ((((InputRGB12Buffer[BufferIndex].DWords[07] >>  16) << 4) | ((InputRGB12Buffer[BufferIndex].DWords[07] >> 28) & 0x00F)) & 0xFFF) * Mul;
 
-	float3 RGB07 = uint3(0,0,0);
-	RGB07.x = ((((InputRGB12Buffer[BufferIndex].DWords[07] >> 24) << 8) | ((InputRGB12Buffer[BufferIndex].DWords[8] >> 00) & 0x0FF)) & 0xFFF) * Mul;
-	RGB07.y = ((((InputRGB12Buffer[BufferIndex].DWords[8] >>  8) << 4) | ((InputRGB12Buffer[BufferIndex].DWords[8] >> 20) & 0x00F)) & 0xFFF) * Mul;
-	RGB07.z = ((((InputRGB12Buffer[BufferIndex].DWords[8] >>  16) << 8) | ((InputRGB12Buffer[BufferIndex].DWords[8] >> 24) & 0x0FF)) & 0xFFF) * Mul;
+	PixelData[7].RGB.x = ((((InputRGB12Buffer[BufferIndex].DWords[07] >> 24) << 8) | ((InputRGB12Buffer[BufferIndex].DWords[8] >> 00) & 0x0FF)) & 0xFFF) * Mul;
+	PixelData[7].RGB.y = ((((InputRGB12Buffer[BufferIndex].DWords[8] >>  8) << 4) | ((InputRGB12Buffer[BufferIndex].DWords[8] >> 20) & 0x00F)) & 0xFFF) * Mul;
+	PixelData[7].RGB.z = ((((InputRGB12Buffer[BufferIndex].DWords[8] >>  16) << 8) | ((InputRGB12Buffer[BufferIndex].DWords[8] >> 24) & 0x0FF)) & 0xFFF) * Mul;
 
+	for (uint PixelIndex = 0; PixelIndex < OutputTexturePixelsPerThread; PixelIndex++)
+	{
 	#if DO_SRGB_TO_LINEAR
-		RGB00 = sRGBToLinear(RGB00);
-		RGB01 = sRGBToLinear(RGB01);
-		RGB02 = sRGBToLinear(RGB02);
-		RGB03 = sRGBToLinear(RGB03);
-		RGB04 = sRGBToLinear(RGB04);
-		RGB05 = sRGBToLinear(RGB05);
-		RGB06 = sRGBToLinear(RGB06);
-		RGB07 = sRGBToLinear(RGB07);
+		PixelData[PixelIndex].RGB.xyz = sRGBToLinear(PixelData[PixelIndex].RGB.xyz);
 	#endif
-
-	OutTexture[uint2(BaseU + 0, BaseV)] = float4(RGB00, 1.0);
-	OutTexture[uint2(BaseU + 1, BaseV)] = float4(RGB01, 1.0);
-	OutTexture[uint2(BaseU + 2, BaseV)] = float4(RGB02, 1.0);
-	OutTexture[uint2(BaseU + 3, BaseV)] = float4(RGB03, 1.0);
-	OutTexture[uint2(BaseU + 4, BaseV)] = float4(RGB04, 1.0);
-	OutTexture[uint2(BaseU + 5, BaseV)] = float4(RGB05, 1.0);
-	OutTexture[uint2(BaseU + 6, BaseV)] = float4(RGB06, 1.0);
-	OutTexture[uint2(BaseU + 7, BaseV)] = float4(RGB07, 1.0);
+		
+		OutTexture[PixelData[PixelIndex].UV] = PixelData[PixelIndex].RGB;
+	}
 }
 
 // Shader to convert from RGBA to RGB Float16
-
-
-
-[numthreads(8,8,1)]
+[numthreads(64,1,1)]
 void RGBToRGB16fBitCS(uint3 ThreadId : SV_DispatchThreadID, uint GroupIndex : SV_GroupIndex)
 {
 	// Make sure it's a valid thread
-	if(ThreadId.x >= HorizontalElementCount || ThreadId.y >= VerticalElementCount)
+	if(ThreadId.x >= TotalElementCount)
 	{
 		return;
 	}
 
-	float X = (ThreadId.x * InputTexturePixelsPerThread) + InputPixelOffsetX;
-	float Y = (ThreadId.y + InputPixelOffsetY) * OnePixelDeltaY;
-
-	float3 RGB0 = InputTexture.SampleLevel(InputSampler, float2(((X + 0.0) + 0.5) * OnePixelDeltaX, Y), 0).rgb;
-	float3 RGB1 = InputTexture.SampleLevel(InputSampler, float2(((X + 1.0) + 0.5) * OnePixelDeltaX, Y), 0).rgb;
-	//RGB0 = float3(1,0,1);
-	//RGB1 = float3(0,1,0);
-
-	// Storage has been untested as no available supported that viewer. Will need to be validated. 
-	uint Index = ThreadId.y * HorizontalElementCount + ThreadId.x;
-
-	// Convert float to half representation and write it in lowest 16 bits of a uint
-	uint IntR0 = f32tof16(RGB0.x);
-	uint IntG0 = f32tof16(RGB0.y);
-	uint IntB0 = f32tof16(RGB0.z);
-	uint IntR1 = f32tof16(RGB1.x);
-	uint IntG1 = f32tof16(RGB1.y);
-	uint IntB1 = f32tof16(RGB1.z);
-
-	#if 1
-	OutRGB16fBuffer[Index].DWords[0]  = ((IntG0 & 0xFF) << 24) | ((IntG0 >> 8) << 16) | ((IntR0 & 0xFF) << 8) | ((IntR0 >> 8) << 0);
-	OutRGB16fBuffer[Index].DWords[1]  = ((IntR1 & 0xFF) << 24) | ((IntR1 >> 8) << 16) | ((IntB0 & 0xFF) << 8) | ((IntB0 >> 8) << 0);
-	OutRGB16fBuffer[Index].DWords[2]  = ((IntB1 & 0xFF) << 24) | ((IntB1 >> 8) << 16) | ((IntG1 & 0xFF) << 8) | ((IntG1 >> 8) << 0);
-	#else
-	OutRGB16fBuffer[Index].DWords[0]  = IntR0;
-	OutRGB16fBuffer[Index].DWords[1]  = IntG0;
-	OutRGB16fBuffer[Index].DWords[2]  = IntB0;
-	#endif
+	// Convert output buffer index to pixels
+	const uint PixelCount = 2;
+	const uint SourcePixelStart = (ThreadId.x * PixelCount);
+	const uint Index = ThreadId.x;
+	
+	uint3 IntRGB[PixelCount];
+	for (uint PixelIndex = 0; PixelIndex < PixelCount; PixelIndex++)
+	{
+		const float BaseX = ((SourcePixelStart + PixelIndex) % CapturedSizeX) + InputPixelOffsetX; 
+		const float BaseY = ((SourcePixelStart + PixelIndex) / CapturedSizeX) + InputPixelOffsetY;
+		const float U = (BaseX + 0.5) * OnePixelDeltaX;
+		const float V = (BaseY + 0.5) * OnePixelDeltaY;
+		float3 RGB = InputTexture.SampleLevel(InputSampler, float2(U,V), 0).rgb; 
+		IntRGB[PixelIndex].x = f32tof16(RGB.x);
+		IntRGB[PixelIndex].y = f32tof16(RGB.y);
+		IntRGB[PixelIndex].z = f32tof16(RGB.z);
+	}
+
+	OutRGB16fBuffer[Index].DWords[0]  = ((IntRGB[0].y & 0xFF) << 24) | ((IntRGB[0].y >> 8) << 16) | ((IntRGB[0].x & 0xFF) << 8) | ((IntRGB[0].x >> 8) << 0);
+	OutRGB16fBuffer[Index].DWords[1]  = ((IntRGB[1].x & 0xFF) << 24) | ((IntRGB[1].x >> 8) << 16) | ((IntRGB[0].z & 0xFF) << 8) | ((IntRGB[0].z >> 8) << 0);
+	OutRGB16fBuffer[Index].DWords[2]  = ((IntRGB[1].z & 0xFF) << 24) | ((IntRGB[1].z >> 8) << 16) | ((IntRGB[1].y & 0xFF) << 8) | ((IntRGB[1].y >> 8) << 0);
 }
 
 /**
@@ -667,11 +578,11 @@ float GetEncodedFloatValue(in uint Value)
 
 // Shader to convert from packed RGB16f to float RGBA texture
 
-[numthreads(8,8,1)]
+[numthreads(64,1,1)]
 void RGB16fBitToRGBACS(uint3 ThreadId : SV_DispatchThreadID, uint GroupIndex : SV_GroupIndex)
 {
 	// Make sure it's a valid thread
-	if(ThreadId.x >= HorizontalElementCount || ThreadId.y >= VerticalElementCount)
+	if(ThreadId.x >= TotalElementCount)
 	{
 		return;
 	}
@@ -679,12 +590,22 @@ void RGB16fBitToRGBACS(uint3 ThreadId : SV_DispatchThreadID, uint GroupIndex : S
 	// Pixels we treat in shader is pretty much hardcoded to 2 so no parameter has been used
 	const uint OutputTexturePixelsPerThread = 2;
 
-	// Scale U (x index) based on the number of pixels we treat per thread
-	const uint BaseU = ThreadId.x * OutputTexturePixelsPerThread;
-	const uint BaseV = ThreadId.y;
+	struct FOutputPixelData
+	{
+		uint2 UV;
+		float4 RGB;
+	};
+	FOutputPixelData PixelData[OutputTexturePixelsPerThread];
 
-	// Input buffer indexing is done in one dimension so adjust using thread Id and elements per row
-	const uint BufferIndex = ThreadId.y * HorizontalElementCount + ThreadId.x;
+	const uint SourcePixelStart = ThreadId.x * OutputTexturePixelsPerThread;
+	for (uint PixelIndex = 0; PixelIndex < OutputTexturePixelsPerThread; PixelIndex++)
+	{
+		PixelData[PixelIndex].UV.x = ((SourcePixelStart + PixelIndex) % OutputTextureSizeX); 
+		PixelData[PixelIndex].UV.y = ((SourcePixelStart + PixelIndex) / OutputTextureSizeX);
+		PixelData[PixelIndex].RGB = float4(0,0,0,1);
+	}
+
+	const uint BufferIndex = ThreadId.x;
 
 	const uint Dword0 = InputRGB16fBuffer[BufferIndex].DWords[00];
 	const uint Dword1 = InputRGB16fBuffer[BufferIndex].DWords[01];
@@ -697,47 +618,68 @@ void RGB16fBitToRGBACS(uint3 ThreadId : SV_DispatchThreadID, uint GroupIndex : S
 	const uint IntGreen1 = Dword2 & 0xFFFF;
 	const uint IntBlue1 = (Dword2 >> 16) & 0xFFFF;
 
-	float4 RGB00 = float4(0,0,0,1);
-	float4 RGB01 = float4(0,0,0,1);
-	RGB00.x = GetEncodedFloatValue(IntRed0);
-	RGB00.y = GetEncodedFloatValue(IntGreen0);
-	RGB00.z = GetEncodedFloatValue(IntBlue0);
-	RGB01.x = GetEncodedFloatValue(IntRed1);
-	RGB01.y = GetEncodedFloatValue(IntGreen1);
-	RGB01.z = GetEncodedFloatValue(IntBlue1);
+	PixelData[0].RGB.x = GetEncodedFloatValue(IntRed0);
+	PixelData[0].RGB.y = GetEncodedFloatValue(IntGreen0);
+	PixelData[0].RGB.z = GetEncodedFloatValue(IntBlue0);
+	PixelData[1].RGB.x = GetEncodedFloatValue(IntRed1);
+	PixelData[1].RGB.y = GetEncodedFloatValue(IntGreen1);
+	PixelData[1].RGB.z = GetEncodedFloatValue(IntBlue1);
 	
-	OutTexture[uint2(BaseU + 0, BaseV)] = RGB00;
-	OutTexture[uint2(BaseU + 1, BaseV)] = RGB01;
+	for (uint PixelIndex = 0; PixelIndex < OutputTexturePixelsPerThread; PixelIndex++)
+	{
+	#if DO_SRGB_TO_LINEAR
+		PixelData[PixelIndex].RGB.xyz = sRGBToLinear(PixelData[PixelIndex].RGB.xyz);
+	#endif
+		
+		OutTexture[PixelData[PixelIndex].UV] = PixelData[PixelIndex].RGB;
+	}
 }
 
 // Shader to convert from YUV 422 8 bits to RGBA
 
-[numthreads(8,8,1)]
+[numthreads(64,1,1)]
 void YUV8Bit422ToRGBACS(uint3 ThreadId : SV_DispatchThreadID, uint GroupIndex : SV_GroupIndex)
 {
 	// Make sure it's a valid thread
-	if(ThreadId.x >= HorizontalElementCount || ThreadId.y >= VerticalElementCount)
+	if(ThreadId.x >= TotalElementCount)
 	{
 		return;
 	}
 
-	// Pixels we treat in shader is pretty much hardcoded to 2 so no parameter has been used
+	// Pixels we treat in shader is pretty much hardcoded to 8 so no parameter has been used
 	const uint OutputTexturePixelsPerThread = 2;
 
-	// Scale U (x index) based on the number of pixels we treat per thread
-	const uint BaseU = ThreadId.x * OutputTexturePixelsPerThread;
-	const uint BaseV = ThreadId.y;
+	struct FOutputPixelData
+	{
+		uint2 UV1;
+		uint2 UV2;
+		float4 YUVY;
+	};
+	const uint PixelContainerSize = OutputTexturePixelsPerThread / 2;
+	FOutputPixelData PixelData[PixelContainerSize];
+
+	const uint SourcePixelStart = ThreadId.x * OutputTexturePixelsPerThread;
+	for (uint ArrayIndex = 0; ArrayIndex < PixelContainerSize; ArrayIndex++)
+	{
+		const uint PixelIndex = ArrayIndex * 2;
+		PixelData[ArrayIndex].UV1.x = ((SourcePixelStart + (PixelIndex + 0)) % OutputTextureSizeX); 
+		PixelData[ArrayIndex].UV1.y = ((SourcePixelStart + (PixelIndex + 0)) / OutputTextureSizeX);
+		PixelData[ArrayIndex].UV2.x = ((SourcePixelStart + (PixelIndex + 1)) % OutputTextureSizeX); 
+		PixelData[ArrayIndex].UV2.y = ((SourcePixelStart + (PixelIndex + 1)) / OutputTextureSizeX);
 
-	// Input buffer indexing is done in one dimension so adjust using thread Id and elements per row
-	const uint BufferIndex = ThreadId.y * HorizontalElementCount + ThreadId.x;
+		PixelData[ArrayIndex].YUVY = float4(0,0,0,0);
+	}
+
+	// Input buffer indexing is simply thread ix
+	const uint BufferIndex = ThreadId.x;
 
 	// Used to normalize pixel values based on 8 bits input
 	const float Mul = 1.0 / 255.0;
 	
-	float U0 = ((InputYUV4228bitBuffer[BufferIndex].DWord0 >>  0) & 0xFF) * Mul;
-	float Y0 = ((InputYUV4228bitBuffer[BufferIndex].DWord0 >>  8) & 0xFF) * Mul;
-	float V0 = ((InputYUV4228bitBuffer[BufferIndex].DWord0 >> 16) & 0xFF) * Mul;
-	float Y1 = ((InputYUV4228bitBuffer[BufferIndex].DWord0 >> 24) & 0xFF) * Mul;
+	PixelData[0].YUVY.y = ((InputYUV4228bitBuffer[BufferIndex].DWord0 >>  0) & 0xFF) * Mul;
+	PixelData[0].YUVY.x = ((InputYUV4228bitBuffer[BufferIndex].DWord0 >>  8) & 0xFF) * Mul;
+	PixelData[0].YUVY.z = ((InputYUV4228bitBuffer[BufferIndex].DWord0 >> 16) & 0xFF) * Mul;
+	PixelData[0].YUVY.w = ((InputYUV4228bitBuffer[BufferIndex].DWord0 >> 24) & 0xFF) * Mul;
 
 	#if DO_SRGB_TO_LINEAR
 	int DoSRGBToLinear = 1;
@@ -745,20 +687,22 @@ void YUV8Bit422ToRGBACS(uint3 ThreadId : SV_DispatchThreadID, uint GroupIndex :
 	int DoSRGBToLinear = 0;
 	#endif
 
-	float3 RGB00 = YuvToRgb(float3(Y0, U0, V0), ColorTransform, DoSRGBToLinear);
-	float3 RGB01 = YuvToRgb(float3(Y1, U0, V0), ColorTransform, DoSRGBToLinear);
-
-	OutTexture[uint2(BaseU + 0, BaseV)] = float4(RGB00, 1.0);
-	OutTexture[uint2(BaseU + 1, BaseV)] = float4(RGB01, 1.0);
+	for (uint PixelIndex = 0; PixelIndex < PixelContainerSize; PixelIndex++)
+	{
+		const float3 RGB00 = YuvToRgb(PixelData[PixelIndex].YUVY.xyz, ColorTransform, DoSRGBToLinear);
+		const float3 RGB01 = YuvToRgb(PixelData[PixelIndex].YUVY.wyz, ColorTransform, DoSRGBToLinear);
+		OutTexture[PixelData[PixelIndex].UV1] = float4(RGB00, 1);
+		OutTexture[PixelData[PixelIndex].UV2] = float4(RGB01, 1);
+	}
 }
 
 // Shader to convert from YUV 422 10 bits packed little endian to RGBA
 
-[numthreads(8,8,1)]
+[numthreads(64,1,1)]
 void YUV10Bit422ToRGBACS(uint3 ThreadId : SV_DispatchThreadID, uint GroupIndex : SV_GroupIndex)
 {
 	// Make sure it's a valid thread
-	if(ThreadId.x >= HorizontalElementCount || ThreadId.y >= VerticalElementCount)
+	if(ThreadId.x >= TotalElementCount)
 	{
 		return;
 	}
@@ -766,32 +710,49 @@ void YUV10Bit422ToRGBACS(uint3 ThreadId : SV_DispatchThreadID, uint GroupIndex :
 	// Pixels we treat in shader is pretty much hardcoded to 8 so no parameter has been used
 	const uint OutputTexturePixelsPerThread = 8;
 
-	// Scale U (x index) based on the number of pixels we treat per thread
-	const uint BaseU = ThreadId.x * OutputTexturePixelsPerThread;
-	const uint BaseV = ThreadId.y;
+	struct FOutputPixelData
+	{
+		uint2 UV1;
+		uint2 UV2;
+		float4 YUVY;
+	};
+	const uint PixelContainerSize = OutputTexturePixelsPerThread / 2;
+	FOutputPixelData PixelData[PixelContainerSize];
+
+	const uint SourcePixelStart = ThreadId.x * OutputTexturePixelsPerThread;
+	for (uint ArrayIndex = 0; ArrayIndex < PixelContainerSize; ArrayIndex++)
+	{
+		const uint PixelIndex = ArrayIndex * 2;
+		PixelData[ArrayIndex].UV1.x = ((SourcePixelStart + (PixelIndex + 0)) % OutputTextureSizeX); 
+		PixelData[ArrayIndex].UV1.y = ((SourcePixelStart + (PixelIndex + 0)) / OutputTextureSizeX);
+		PixelData[ArrayIndex].UV2.x = ((SourcePixelStart + (PixelIndex + 1)) % OutputTextureSizeX); 
+		PixelData[ArrayIndex].UV2.y = ((SourcePixelStart + (PixelIndex + 1)) / OutputTextureSizeX);
+
+		PixelData[ArrayIndex].YUVY = float4(0,0,0,0);
+	}
 
-	// Input buffer indexing is done in one dimension so adjust using thread Id and elements per row
-	const uint BufferIndex = ThreadId.y * HorizontalElementCount + ThreadId.x;
+	// Input buffer indexing is simply thread ix
+	const uint BufferIndex = ThreadId.x;
 
 	// Used to normalize pixel values based on 10 bits input
 	const float Mul = 1.0 / 1023.0;
 
-	float Y0 = ((((InputYCbCrBuffer[BufferIndex].DWord0 >>  8)  << 4) | ((InputYCbCrBuffer[BufferIndex].DWord0 >> 20) & 0x00F)) & 0x3FF) * Mul;
-	float U0 = ((((InputYCbCrBuffer[BufferIndex].DWord0 >>  0)  << 2) | ((InputYCbCrBuffer[BufferIndex].DWord0 >> 14) & 0x003)) & 0x3FF) * Mul;
-	float V0 = ((((InputYCbCrBuffer[BufferIndex].DWord0 >> 16)  << 6) | ((InputYCbCrBuffer[BufferIndex].DWord0 >> 26) & 0x03F)) & 0x3FF) * Mul;
-	float Y1 = ((((InputYCbCrBuffer[BufferIndex].DWord0 >> 24)  << 8) | ((InputYCbCrBuffer[BufferIndex].DWord1 >> 00) & 0x0FF)) & 0x3FF) * Mul;
-	float Y2 = ((((InputYCbCrBuffer[BufferIndex].DWord1 >> 16)  << 4) | ((InputYCbCrBuffer[BufferIndex].DWord1 >> 28) & 0x00F)) & 0x3FF) * Mul;
-	float U2 = ((((InputYCbCrBuffer[BufferIndex].DWord1 >>  8)  << 2) | ((InputYCbCrBuffer[BufferIndex].DWord1 >> 22) & 0x003)) & 0x3FF) * Mul;
-	float V2 = ((((InputYCbCrBuffer[BufferIndex].DWord1 >> 24)  << 6) | ((InputYCbCrBuffer[BufferIndex].DWord2 >> 02) & 0x03F)) & 0x3FF) * Mul;
-	float Y3 = ((((InputYCbCrBuffer[BufferIndex].DWord2 >> 00)  << 8) | ((InputYCbCrBuffer[BufferIndex].DWord2 >>  8) & 0x0FF)) & 0x3FF) * Mul;
-	float Y4 = ((((InputYCbCrBuffer[BufferIndex].DWord2 >> 24)  << 4) | ((InputYCbCrBuffer[BufferIndex].DWord3 >> 04) & 0x00F)) & 0x3FF) * Mul;
-	float U4 = ((((InputYCbCrBuffer[BufferIndex].DWord2 >> 16)  << 2) | ((InputYCbCrBuffer[BufferIndex].DWord2 >> 30) & 0x003)) & 0x3FF) * Mul;
-	float V4 = ((((InputYCbCrBuffer[BufferIndex].DWord3 >> 00)  << 6) | ((InputYCbCrBuffer[BufferIndex].DWord3 >> 10) & 0x03F)) & 0x3FF) * Mul;
-	float Y5 = ((((InputYCbCrBuffer[BufferIndex].DWord3 >>  8)  << 8) | ((InputYCbCrBuffer[BufferIndex].DWord3 >> 16) & 0x0FF)) & 0x3FF) * Mul;
-	float Y6 = ((((InputYCbCrBuffer[BufferIndex].DWord4 >>  0)  << 4) | ((InputYCbCrBuffer[BufferIndex].DWord3 >> 12) & 0x00F)) & 0x3FF) * Mul;
-	float U6 = ((((InputYCbCrBuffer[BufferIndex].DWord3 >> 24)  << 2) | ((InputYCbCrBuffer[BufferIndex].DWord4 >> 06) & 0x003)) & 0x3FF) * Mul;
-	float V6 = ((((InputYCbCrBuffer[BufferIndex].DWord4 >>  8)  << 6) | ((InputYCbCrBuffer[BufferIndex].DWord4 >> 18) & 0x03F)) & 0x3FF) * Mul;
-	float Y7 = ((((InputYCbCrBuffer[BufferIndex].DWord4 >> 16)  << 8) | ((InputYCbCrBuffer[BufferIndex].DWord4 >> 24) & 0x0FF)) & 0x3FF) * Mul;
+	PixelData[0].YUVY.x = ((((InputYCbCrBuffer[BufferIndex].DWord0 >>  8)  << 4) | ((InputYCbCrBuffer[BufferIndex].DWord0 >> 20) & 0x00F)) & 0x3FF) * Mul;
+	PixelData[0].YUVY.y = ((((InputYCbCrBuffer[BufferIndex].DWord0 >>  0)  << 2) | ((InputYCbCrBuffer[BufferIndex].DWord0 >> 14) & 0x003)) & 0x3FF) * Mul;
+	PixelData[0].YUVY.z = ((((InputYCbCrBuffer[BufferIndex].DWord0 >> 16)  << 6) | ((InputYCbCrBuffer[BufferIndex].DWord0 >> 26) & 0x03F)) & 0x3FF) * Mul;
+	PixelData[0].YUVY.w = ((((InputYCbCrBuffer[BufferIndex].DWord0 >> 24)  << 8) | ((InputYCbCrBuffer[BufferIndex].DWord1 >> 00) & 0x0FF)) & 0x3FF) * Mul;
+	PixelData[1].YUVY.x = ((((InputYCbCrBuffer[BufferIndex].DWord1 >> 16)  << 4) | ((InputYCbCrBuffer[BufferIndex].DWord1 >> 28) & 0x00F)) & 0x3FF) * Mul;
+	PixelData[1].YUVY.y = ((((InputYCbCrBuffer[BufferIndex].DWord1 >>  8)  << 2) | ((InputYCbCrBuffer[BufferIndex].DWord1 >> 22) & 0x003)) & 0x3FF) * Mul;
+	PixelData[1].YUVY.z = ((((InputYCbCrBuffer[BufferIndex].DWord1 >> 24)  << 6) | ((InputYCbCrBuffer[BufferIndex].DWord2 >> 02) & 0x03F)) & 0x3FF) * Mul;
+	PixelData[1].YUVY.w = ((((InputYCbCrBuffer[BufferIndex].DWord2 >> 00)  << 8) | ((InputYCbCrBuffer[BufferIndex].DWord2 >>  8) & 0x0FF)) & 0x3FF) * Mul;
+	PixelData[2].YUVY.x = ((((InputYCbCrBuffer[BufferIndex].DWord2 >> 24)  << 4) | ((InputYCbCrBuffer[BufferIndex].DWord3 >> 04) & 0x00F)) & 0x3FF) * Mul;
+	PixelData[2].YUVY.y = ((((InputYCbCrBuffer[BufferIndex].DWord2 >> 16)  << 2) | ((InputYCbCrBuffer[BufferIndex].DWord2 >> 30) & 0x003)) & 0x3FF) * Mul;
+	PixelData[2].YUVY.z = ((((InputYCbCrBuffer[BufferIndex].DWord3 >> 00)  << 6) | ((InputYCbCrBuffer[BufferIndex].DWord3 >> 10) & 0x03F)) & 0x3FF) * Mul;
+	PixelData[2].YUVY.w = ((((InputYCbCrBuffer[BufferIndex].DWord3 >>  8)  << 8) | ((InputYCbCrBuffer[BufferIndex].DWord3 >> 16) & 0x0FF)) & 0x3FF) * Mul;
+	PixelData[3].YUVY.x = ((((InputYCbCrBuffer[BufferIndex].DWord4 >>  0)  << 4) | ((InputYCbCrBuffer[BufferIndex].DWord3 >> 12) & 0x00F)) & 0x3FF) * Mul;
+	PixelData[3].YUVY.y = ((((InputYCbCrBuffer[BufferIndex].DWord3 >> 24)  << 2) | ((InputYCbCrBuffer[BufferIndex].DWord4 >> 06) & 0x003)) & 0x3FF) * Mul;
+	PixelData[3].YUVY.z = ((((InputYCbCrBuffer[BufferIndex].DWord4 >>  8)  << 6) | ((InputYCbCrBuffer[BufferIndex].DWord4 >> 18) & 0x03F)) & 0x3FF) * Mul;
+	PixelData[3].YUVY.w = ((((InputYCbCrBuffer[BufferIndex].DWord4 >> 16)  << 8) | ((InputYCbCrBuffer[BufferIndex].DWord4 >> 24) & 0x0FF)) & 0x3FF) * Mul;
 
 	#if DO_SRGB_TO_LINEAR
 	int DoSRGBToLinear = 1;
@@ -799,21 +760,11 @@ void YUV10Bit422ToRGBACS(uint3 ThreadId : SV_DispatchThreadID, uint GroupIndex :
 	int DoSRGBToLinear = 0;
 	#endif
 
-	float3 RGB00 = YuvToRgb(float3(Y0, U0, V0), ColorTransform, DoSRGBToLinear);
-	float3 RGB01 = YuvToRgb(float3(Y1, U0, V0), ColorTransform, DoSRGBToLinear);
-	float3 RGB02 = YuvToRgb(float3(Y2, U2, V2), ColorTransform, DoSRGBToLinear);
-	float3 RGB03 = YuvToRgb(float3(Y3, U2, V2), ColorTransform, DoSRGBToLinear);
-	float3 RGB04 = YuvToRgb(float3(Y4, U4, V4), ColorTransform, DoSRGBToLinear);
-	float3 RGB05 = YuvToRgb(float3(Y5, U4, V4), ColorTransform, DoSRGBToLinear);
-	float3 RGB06 = YuvToRgb(float3(Y6, U6, V6), ColorTransform, DoSRGBToLinear);
-	float3 RGB07 = YuvToRgb(float3(Y7, U6, V6), ColorTransform, DoSRGBToLinear);
-
-	OutTexture[uint2(BaseU + 0, BaseV)] = float4(RGB00, 1.0);
-	OutTexture[uint2(BaseU + 1, BaseV)] = float4(RGB01, 1.0);
-	OutTexture[uint2(BaseU + 2, BaseV)] = float4(RGB02, 1.0);
-	OutTexture[uint2(BaseU + 3, BaseV)] = float4(RGB03, 1.0);
-	OutTexture[uint2(BaseU + 4, BaseV)] = float4(RGB04, 1.0);
-	OutTexture[uint2(BaseU + 5, BaseV)] = float4(RGB05, 1.0);
-	OutTexture[uint2(BaseU + 6, BaseV)] = float4(RGB06, 1.0);
-	OutTexture[uint2(BaseU + 7, BaseV)] = float4(RGB07, 1.0);
+	for (uint PixelIndex = 0; PixelIndex < PixelContainerSize; PixelIndex++)
+	{
+		const float3 RGB00 = YuvToRgb(PixelData[PixelIndex].YUVY.xyz, ColorTransform, DoSRGBToLinear);
+		const float3 RGB01 = YuvToRgb(PixelData[PixelIndex].YUVY.wyz, ColorTransform, DoSRGBToLinear);
+		OutTexture[PixelData[PixelIndex].UV1] = float4(RGB00, 1);
+		OutTexture[PixelData[PixelIndex].UV2] = float4(RGB01, 1);
+	}
 }
diff --git a/Engine/Plugins/VirtualProduction/Rivermax/RivermaxCore/Source/RivermaxCore/Private/RTPHeader.cpp b/Engine/Plugins/VirtualProduction/Rivermax/RivermaxCore/Source/RivermaxCore/Private/RTPHeader.cpp
new file mode 100644
--- /dev/null
+++ b/Engine/Plugins/VirtualProduction/Rivermax/RivermaxCore/Source/RivermaxCore/Private/RTPHeader.cpp
@@ -0,0 +1,154 @@
+// Copyright Epic Games, Inc. All Rights Reserved.
+
+#include "RTPHeader.h"
+
+#include "RivermaxLog.h"
+
+namespace UE::RivermaxCore::Private
+{
+	uint16 FRawRTPHeader::GetSrd1RowNumber() const
+	{
+		return ((SRD1RowNumberHigh << 8) | SRD1RowNumberLow);
+	}
+
+	void FRawRTPHeader::SetSrd1RowNumber(uint16 RowNumber)
+	{
+		SRD1RowNumberHigh = (RowNumber >> 8) & 0xFF;
+		SRD1RowNumberLow = RowNumber & 0xFF;
+	}
+
+	uint16 FRawRTPHeader::GetSrd1Offset() const
+	{
+		return ((SRD1OffsetHigh << 8) | SRD1OffsetLow);
+	}
+
+	void FRawRTPHeader::SetSrd1Offset(uint16 Offset)
+	{
+		SRD1OffsetHigh = (Offset >> 8) & 0xFF;
+		SRD1OffsetLow = Offset & 0xFF;
+	}
+
+	uint16 FRawRTPHeader::GetSrd2RowNumber() const
+	{
+		return ((SRD2RowNumberHigh << 8) | SRD2RowNumberLow);
+	}
+
+	void FRawRTPHeader::SetSrd2RowNumber(uint16 RowNumber)
+	{
+		SRD2RowNumberHigh = (RowNumber >> 8) & 0xFF;
+		SRD2RowNumberLow = RowNumber & 0xFF;
+	}
+
+	uint16 FRawRTPHeader::GetSrd2Offset() const
+	{
+		return ((SRD2OffsetHigh << 8) | SRD2OffsetLow);
+	}
+
+	void FRawRTPHeader::SetSrd2Offset(uint16 Offset)
+	{
+		SRD2OffsetHigh = (Offset >> 8) & 0xFF;
+		SRD2OffsetLow = Offset & 0xFF;
+	}
+
+	uint8* GetRTPHeaderPointer(uint8* InHeader)
+	{
+		check(InHeader);
+
+		static constexpr uint32 ETH_TYPE_802_1Q = 0x8100;          /* 802.1Q VLAN Extended Header  */
+		static constexpr uint32 RTP_HEADER_SIZE = 12;
+		uint16* ETHProto = (uint16_t*)(InHeader + RTP_HEADER_SIZE);
+		if (ETH_TYPE_802_1Q == ByteSwap(*ETHProto))
+		{
+			InHeader += 46; // 802 + 802.1Q + IP + UDP
+		}
+		else
+		{
+			InHeader += 42; // 802 + IP + UDP
+		}
+		return InHeader;
+	}
+
+	FRTPHeader::FRTPHeader(const FRawRTPHeader& RawHeader)
+	{
+		Timestamp = 0;
+
+		if (RawHeader.Version != 2)
+		{
+			return;
+		}
+
+		// Pretty sure some data needs to be swapped but can't validate that until we have other hardware generating data
+		SequencerNumber = (ByteSwap((uint16)RawHeader.ExtendedSequenceNumber) << 16) | ByteSwap((uint16)RawHeader.SequenceNumber);
+		Timestamp = ByteSwap(RawHeader.Timestamp);
+		bIsMarkerBit = RawHeader.MarkerBit;
+
+		SyncSouceId = RawHeader.SynchronizationSource;
+
+		SRD1.Length = ByteSwap((uint16)RawHeader.SRD1Length);
+		SRD1.DataOffset = RawHeader.GetSrd1Offset();
+		SRD1.RowNumber = RawHeader.GetSrd1RowNumber();
+		SRD1.bIsFieldOne = RawHeader.FieldIdentification1;
+		SRD1.bHasContinuation = RawHeader.ContinuationBit1;
+
+		if (SRD1.bHasContinuation)
+		{
+			SRD2.Length = ByteSwap((uint16)RawHeader.SRD2Length);
+			SRD2.DataOffset = RawHeader.GetSrd2Offset();
+			SRD2.RowNumber = RawHeader.GetSrd2RowNumber();
+			SRD2.bIsFieldOne = RawHeader.FieldIdentification2;
+			SRD2.bHasContinuation = RawHeader.ContinuationBit2;
+
+			if (SRD2.bHasContinuation == true)
+			{
+				UE_LOG(LogRivermax, Verbose, TEXT("Received SRD with more than 2 SRD which isn't supported."));
+			}
+		}
+	}
+
+	uint16 FRTPHeader::GetTotalPayloadSize() const
+	{
+		uint16 PayloadSize = SRD1.Length;
+		if (SRD1.bHasContinuation)
+		{
+			PayloadSize += SRD2.Length;
+		}
+
+		return PayloadSize;
+	}
+
+	uint16 FRTPHeader::GetLastPayloadSize() const
+	{
+		if (SRD1.bHasContinuation)
+		{
+			return SRD2.Length;
+		}
+
+		return SRD1.Length;
+	}
+
+	uint16 FRTPHeader::GetLastRowOffset() const
+	{
+		if (SRD1.bHasContinuation)
+		{
+			return SRD2.DataOffset;
+		}
+
+		return SRD1.DataOffset;
+	}
+
+	uint16 FRTPHeader::GetLastRowNumber() const
+	{
+		if (SRD1.bHasContinuation)
+		{
+			return SRD2.RowNumber;
+		}
+
+		return SRD1.RowNumber;
+	}
+}
+
+
+
+
+
+
diff --git a/Engine/Plugins/VirtualProduction/Rivermax/RivermaxCore/Source/RivermaxCore/Private/RTPHeader.h b/Engine/Plugins/VirtualProduction/Rivermax/RivermaxCore/Source/RivermaxCore/Private/RTPHeader.h
new file mode 100644
--- /dev/null
+++ b/Engine/Plugins/VirtualProduction/Rivermax/RivermaxCore/Source/RivermaxCore/Private/RTPHeader.h
@@ -0,0 +1,160 @@
+// Copyright Epic Games, Inc. All Rights Reserved.
+
+#pragma once
+
+#include "CoreTypes.h"
+
+
+namespace UE::RivermaxCore::Private
+{
+	//RTP Header used for 2110 following https://www.rfc-editor.org/rfc/rfc4175.html
+
+	/* RTP Header -  12 bytes
+	0                   1                   2                   3
+	0 1 2 3 4 5 6 7 8 9 0 1 2 3 4 5 6 7 8 9 0 1 2 3 4 5 6 7 8 9 0 1
+	+-+-+-+-+-+-+-+-+-+-+-+-+-+-+-+-+-+-+-+-+-+-+-+-+-+-+-+-+-+-+-+-+
+	| V |P|X|  CC   |M|     PT      |            SEQ                |
+	+-+-+-+-+-+-+-+-+-+-+-+-+-+-+-+-+-+-+-+-+-+-+-+-+-+-+-+-+-+-+-+-+
+	|                           timestamp                           |
+	+-+-+-+-+-+-+-+-+-+-+-+-+-+-+-+-+-+-+-+-+-+-+-+-+-+-+-+-+-+-+-+-+
+	|                           ssrc                                |
+	+-+-+-+-+-+-+-+-+-+-+-+-+-+-+-+-+-+-+-+-+-+-+-+-+-+-+-+-+-+-+-+-+
+
+	SRD Header 8-14 bytes
+
+	0                   1                   2                   3
+	0 1 2 3 4 5 6 7 8 9 0 1 2 3 4 5 6 7 8 9 0 1 2 3 4 5 6 7 8 9 0 1
+	+-+-+-+-+-+-+-+-+-+-+-+-+-+-+-+-+-+-+-+-+-+-+-+-+-+-+-+-+-+-+-+-+
+	|    Extended Sequence Number   |           SRD Length          |
+	+-+-+-+-+-+-+-+-+-+-+-+-+-+-+-+-+-+-+-+-+-+-+-+-+-+-+-+-+-+-+-+-+
+	|F|     SRD Row Number          |C|         SRD Offset          |
+	+-+-+-+-+-+-+-+-+-+-+-+-+-+-+-+-+-+-+-+-+-+-+-+-+-+-+-+-+-+-+-+-+ */
+
+	/** Raw representation as it's built for the network */
+
+/** @note When other platform than windows are supports, reverify support for pragma_pack and endianness */
+#if PLATFORM_SUPPORTS_PRAGMA_PACK
+#pragma pack(push, 1)
+#endif
+	struct FRawRTPHeader
+	{
+		uint32 ContributingSourceCount : 4;
+		uint32 ExtensionBit : 1;
+		uint32 PaddingBit : 1;
+		uint32 Version : 2;
+		uint32 PayloadType : 7;
+		uint32 MarkerBit : 1;
+		uint32 SequenceNumber : 16;
+		uint32 Timestamp : 32;
+		uint32 SynchronizationSource : 32;
+		uint32 ExtendedSequenceNumber : 16;
+		//SRD 1
+		uint32 SRD1Length : 16;
+		uint32 SRD1RowNumberHigh : 7;
+		uint32 FieldIdentification1 : 1;
+		uint32 SRD1RowNumberLow : 8;
+		uint32 SRD1OffsetHigh : 7;
+		uint32 ContinuationBit1 : 1;
+		uint32 SRD1OffsetLow : 8;
+		//SRD 2
+		uint32 SRD2Length : 16;
+		uint32 SRD2RowNumberHigh : 7;
+		uint32 FieldIdentification2 : 1;
+		uint32 SRD2RowNumberLow : 8;
+		uint32 SRD2OffsetHigh : 7;
+		uint32 ContinuationBit2 : 1;
+		uint32 SRD2OffsetLow : 8;
+
+		/** Returns SRD1 associated row number */
+		uint16 GetSrd1RowNumber() const;
+
+		/** Sets SRD1 associated row number */
+		void SetSrd1RowNumber(uint16 RowNumber);
+
+		/** Returns SRD1 pixel offset in its associated row */
+		uint16 GetSrd1Offset() const;
+
+		/** Sets SRD1 pixel offset in its associated row */
+		void SetSrd1Offset(uint16 Offset);
+
+		/** Sets SRD2 associated row number */
+		uint16 GetSrd2RowNumber() const;
+
+		/** Sets SRD1 associated row number */
+		void SetSrd2RowNumber(uint16 RowNumber);
+
+		/** Returns SRD1 pixel offset in its associated row */
+		uint16 GetSrd2Offset() const;
+
+		/** Sets SRD2 associated row number */
+		void SetSrd2Offset(uint16 Offset);
+
+		/** Size of RTP representation whether it has one or two SRDs */
+		static constexpr uint32 OneSRDSize = 20;
+		static constexpr uint32 TwoSRDSize = 26;
+	};
+#if PLATFORM_SUPPORTS_PRAGMA_PACK
+#pragma pack(pop)
+#endif
+
+	struct FSRDHeader
+	{
+		/** Length of payload. Is a multiple of pgroup (see pixel formats) */
+		uint16 Length = 0;
+
+		/** False if progressive or first field of interlace. True if second field of interlace */
+		bool bIsFieldOne = false;
+
+		/** Video line number, starts at 0 */
+		uint16 RowNumber = 0;
+
+		/** Whether another SRD is following this one */
+		bool bHasContinuation = false;
+
+		/** Location of the first pixel in payload, in pixel */
+		uint16 DataOffset = 0;
+	};
+
+	/** RTP header built from network representation not requiring any byte swapping */
+	struct FRTPHeader
+	{
+		FRTPHeader() = default;
+		FRTPHeader(const FRawRTPHeader& RawHeader);
+
+		/** Returns the total payload of this RTP */
+		uint16 GetTotalPayloadSize() const;
+		
+		/** Returns the payload size of the last SRD in this RTP */
+		uint16 GetLastPayloadSize() const;
+
+		/** Returns the row offset of the last SRD in this RTP */
+		uint16 GetLastRowOffset() const;
+
+		/** Returns the row number of the last SRD in this RTP */
+		uint16 GetLastRowNumber() const;
+
+		/** Sequence number including extension if present */
+		uint32 SequencerNumber = 0;
+
+		/** Timestamp of frame in the specified clock resolution. Video is typically 90kHz */
+		uint32 Timestamp = 0;
+
+		/** Identification of this stream */
+		uint32 SyncSouceId = 0;
+
+		/** Whether extensions (SRD headers) are present */
+		bool bHasExtension = false;
+
+		/** True if RTP packet is last of video stream */
+		bool bIsMarkerBit = false;
+
+		/** Only supports 2 SRD for now. Adjust if needed */
+		FSRDHeader SRD1;
+		FSRDHeader SRD2;
+	};
+
+	/** Returns RTPHeader pointer from a raw ethernet packet skipping 802, IP, UDP headers */
+	uint8* GetRTPHeaderPointer(uint8* InHeader);
+}
+
+
diff --git a/Engine/Plugins/VirtualProduction/Rivermax/RivermaxCore/Source/RivermaxCore/Private/RivermaxInputStream.cpp b/Engine/Plugins/VirtualProduction/Rivermax/RivermaxCore/Source/RivermaxCore/Private/RivermaxInputStream.cpp
--- a/Engine/Plugins/VirtualProduction/Rivermax/RivermaxCore/Source/RivermaxCore/Private/RivermaxInputStream.cpp
+++ b/Engine/Plugins/VirtualProduction/Rivermax/RivermaxCore/Source/RivermaxCore/Private/RivermaxInputStream.cpp
@@ -10,6 +10,7 @@
 #include "Misc/ByteSwap.h"
 #include "RivermaxLog.h"
 #include "RivermaxUtils.h"
+#include "RTPHeader.h"
 
 #if PLATFORM_WINDOWS
 #include <WS2tcpip.h>
@@ -29,117 +30,82 @@ namespace UE::RivermaxCore::Private
 		TEXT("Expected payload size used to initialize rivermax stream."),
 		ECVF_Default);
 
-
-#if PLATFORM_SUPPORTS_PRAGMA_PACK
-#pragma pack(push, 1)
-#endif
-	struct FRawRTPHeader
+	/** 
+	 * Converts a timestamp in MediaClock period units to a frame number for a given frame rate 
+	 * 2110-20 streams uses a standard media clock rate of 90kHz
+	 */
+	uint32 TimestampToFrameNumber(uint32 Timestamp, const FFrameRate& FrameRate)
 	{
-		uint32 ContributingSourceCount : 4;
-		uint32 ExtensionBit : 1;
-		uint32 PaddingBit : 1;
-		uint32 Version : 2;
-		uint32 PayloadType : 7;
-		uint32 MarkerBit : 1;
-		uint32 SequenceNumber : 16;
-		uint32 Timestamp : 32;
-		uint32 SynchronizationSource : 32;
-		uint32 ExtendedSequenceNumber : 16;
-		//SRD 1
-		uint32 SRDLength1 : 16;
-		uint32 SRDRowNumberHigh1 : 7;
-		uint32 FieldIdentification1 : 1;
-		uint32 SRDRowNumberLow1 : 8;
-		uint32 SRDOffsetHigh1 : 7;
-		uint32 ContinuationBit1 : 1;
-		uint32 SRDOffsetLow1 : 8;
-		//SRD 2
-		uint32 SRDLength2 : 16;
-		uint32 SRDRowNumberHigh2 : 7;
-		uint32 FieldIdentification2 : 1;
-		uint32 SRDRowNumberLow2 : 8;
-		uint32 SRDOffsetHigh2 : 7;
-		uint32 ContinuationBit2 : 1;
-		uint32 SRDOffsetLow2 : 8;
-
-		uint16 GetSrd1RowNumber() const { return ((SRDRowNumberHigh1 << 8) | SRDRowNumberLow1); }
-		uint16 GetSrd1Offset() const { return ((SRDOffsetHigh1 << 8) | SRDOffsetLow1); }
-
-		uint16 GetSrd2RowNumber() const { return ((SRDRowNumberHigh2 << 8) | SRDRowNumberLow2); }
-		uint16 GetSrd2Offset() const { return ((SRDOffsetHigh2 << 8) | SRDOffsetLow2); }
-	};
-#if PLATFORM_SUPPORTS_PRAGMA_PACK
-#pragma pack(pop)
-#endif
+		using namespace UE::RivermaxCore::Private::Utils;
+		const double MediaFrameTime = Timestamp / MediaClockSampleRate;
+		const uint32 FrameNumber = FMath::Floor(MediaFrameTime * FrameRate.AsDecimal());
+		return FrameNumber;
+	}
 
-	struct FSRDHeader
+	void FFrameDescriptionTrackingData::ResetSingleFrameTracking()
 	{
-		/** Length of payload. Is a multiple of pgroup (see pixel formats) */
-		uint16 Length = 0;
-
-		/** False if progressive or first field of interlace. True if second field of interlace */
-		bool bIsFieldOne = false;
-
-		/** Video line number, starts at 0 */
-		uint16 RowNumber = 0;
+		PayloadSizeReceived.Empty();
+	}
 
-		/** Whether another SRD is following this one */
-		bool bHasContinuation = false;
-		
-		/** Location of the first pixel in payload, in pixel */
-		uint16 DataOffset = 0;
-	};
-	
-	struct FRTPHeader
+	void FFrameDescriptionTrackingData::EvaluateNewRTP(const FRTPHeader& NewHeader)
 	{
-		/** Sequence number including extension if present */
-		uint32 SequencerNumber = 0;
-
-		/** Timestamp of frame in the specified clock resolution. Video is typically 90kHz */
-		uint32 Timestamp = 0;
-
-		/** Identification of this stream */
-		uint32 SyncSouceId = 0;
-
-		/** Whether extensions (SRD headers) are present */
-		bool bHasExtension = false;
+		UpdateResolutionDetection(NewHeader);
+		UpdatePayloadSizeTracking(NewHeader);
 
-		/** True if RTP packet is last of video stream */
-		bool bIsMarkerBit = false;
-
-		/** Only supports 2 SRD for now. Adjust if needed */
-		FSRDHeader SRD1;
-		FSRDHeader SRD2;
-	};
+		if (NewHeader.bIsMarkerBit)
+		{
+			ResetSingleFrameTracking();
+		}
+	}
 
-	uint8* GetRTPHeaderPointer(uint8* InHeader)
+	void FFrameDescriptionTrackingData::UpdateResolutionDetection(const FRTPHeader& NewHeader)
 	{
-		check(InHeader);
-
-		static constexpr uint32 ETH_TYPE_802_1Q = 0x8100;          /* 802.1Q VLAN Extended Header  */
-		static constexpr uint32 RTP_HEADER_SIZE = 12;
-		uint16* ETHProto = (uint16_t*)(InHeader + RTP_HEADER_SIZE);
-		if (ETH_TYPE_802_1Q == ByteSwap(*ETHProto))
+		if (NewHeader.bIsMarkerBit)
 		{
-			InHeader += 46; // 802 + 802.1Q + IP + UDP
+			const FVideoFormatInfo ExpectedFormatInfo = FStandardVideoFormat::GetVideoFormatInfo(ExpectedSamplingType);
+			const uint16 LastPayloadSize = NewHeader.GetLastPayloadSize();
+			if (LastPayloadSize % ExpectedFormatInfo.PixelGroupSize == 0)
+			{
+				DetectedResolution.X = NewHeader.GetLastRowOffset() + (LastPayloadSize / ExpectedFormatInfo.PixelGroupSize * ExpectedFormatInfo.PixelGroupCoverage);
+				DetectedResolution.Y = NewHeader.GetLastRowNumber() + 1;
+				
+				// Renable logging when a valid incoming frame was detected
+				bHasLoggedSamplingWarning = false;
+			}
+			else if(!bHasLoggedSamplingWarning)
+			{
+				bHasLoggedSamplingWarning = true;
+				UE_LOG(LogRivermax, Warning, TEXT("Detected incoming signal with unexpected sampling type."));
+			}
 		}
-		else
+	}
+
+	void FFrameDescriptionTrackingData::UpdatePayloadSizeTracking(const FRTPHeader& NewHeader)
+	{
+		if (NewHeader.SRD1.Length > 0)
 		{
-			InHeader += 42; // 802 + IP + UDP
+			uint32 PayloadSize = NewHeader.SRD1.Length;
+			if (NewHeader.SRD1.bHasContinuation)
+			{
+				PayloadSize += NewHeader.SRD2.Length;
+			}
+
+			if (PayloadSizeReceived.IsEmpty())
+			{
+				CommonPayloadSize = PayloadSize;
+			}
+			else if (!NewHeader.bIsMarkerBit && PayloadSize != CommonPayloadSize)
+			{
+				CommonPayloadSize = -1;
+			}
+
+			PayloadSizeReceived.FindOrAdd(PayloadSize) += 1;
 		}
-		return InHeader;
 	}
 
-	/** 
-	 * Converts a timestamp in MediaClock period units to a frame number for a given frame rate 
-	 * 2110-20 streams uses a standard media clock rate of 90kHz
-	 */
-	uint32 TimestampToFrameNumber(uint32 Timestamp, const FFrameRate& FrameRate)
+	bool FFrameDescriptionTrackingData::HasDetectedValidResolution() const
 	{
-		using namespace UE::RivermaxCore::Private::Utils;
-		const double MediaFrameTime = Timestamp / MediaClockSampleRate;
-		const uint32 FrameNumber = FMath::Floor(MediaFrameTime * FrameRate.AsDecimal());
-		return FrameNumber;
+		return DetectedResolution.X > 0 && DetectedResolution.Y > 0;
 	}
 
 	FRivermaxInputStream::FRivermaxInputStream()
@@ -165,6 +131,15 @@ namespace UE::RivermaxCore::Private
 		Listener = &InListener;
 		FormatInfo = FStandardVideoFormat::GetVideoFormatInfo(Options.PixelFormat);
 		ExpectedPayloadSize = CVarExpectedPayloadSize.GetValueOnGameThread();
+		bIsDynamicHeaderEnabled = RivermaxModule.GetRivermaxManager()->EnableDynamicHeaderSupport(Options.InterfaceAddress);
+
+		if (Options.bEnforceVideoFormat)
+		{
+			StreamResolution = Options.EnforcedResolution;
+		}
+		
+		// For now, we don't try to detect pixel format so use what the user selected
+		FrameDescriptionTracking.ExpectedSamplingType = Options.PixelFormat;
 
 		InitTaskFuture = Async(EAsyncExecution::TaskGraph, [this]()
 		{
@@ -192,6 +167,8 @@ namespace UE::RivermaxCore::Private
 				//Configure Flow and destination IP (multicast)
 				memset(&FlowAttribute, 0, sizeof(FlowAttribute));
 				FlowAttribute.local_addr.sin_family = AF_INET;
+				FlowAttribute.remote_addr.sin_family = AF_INET;
+
 				FlowAttribute.flow_id = FlowId;
 				if (inet_pton(AF_INET, StringCast<ANSICHAR>(*Options.StreamAddress).Get(), &FlowAttribute.local_addr.sin_addr) != 1)
 				{
@@ -199,9 +176,15 @@ namespace UE::RivermaxCore::Private
 				}
 				else
 				{
+					RivermaxInterface.sin_port = Options.Port;
 					FlowAttribute.local_addr.sin_port = ByteSwap((uint16)Options.Port);
 
-					const rmax_in_buffer_attr_flags_t BufferAttributeFlags = RMAX_IN_BUFFER_ATTER_FLAG_NONE; //todo whether ordering is based on sequence or extended sequence
+					rmax_in_buffer_attr_flags_t BufferAttributeFlags = RMAX_IN_BUFFER_ATTER_FLAG_NONE;
+					if (bIsDynamicHeaderEnabled)
+					{
+						BufferAttributeFlags = RMAX_IN_BUFFER_ATTR_BUFFER_RTP_SMPTE_2110_20_DYNAMIC_HDS;
+					}
+					;
 					uint32 BufferElement = 1 << 18;//todo number of packets to allocate memory for
 					rmax_in_buffer_attr BufferAttributes;
 					FMemory::Memset(&BufferAttributes, 0, sizeof(BufferAttributes));
@@ -236,9 +219,10 @@ namespace UE::RivermaxCore::Private
 								RivermaxThread = FRunnableThread::Create(this, TEXT("Rivermax InputStream Thread"), 128 * 1024, TPri_AboveNormal, FPlatformAffinity::GetPoolThreadMask());
 								bWasSuccessful = true;
 
-								UE_LOG(LogRivermax, Display, TEXT("Input started receiving %dx%d%s")
-									, Options.AlignedResolution.X
-									, Options.AlignedResolution.Y
+								UE_LOG(LogRivermax, Display, TEXT("Input stream started listening to stream %s:%d on interface %s%s")
+									, *Options.StreamAddress
+									, Options.Port
+									, *Options.InterfaceAddress
 									, bIsUsingGPUDirect ? TEXT(" using GPUDirect") : TEXT(""));
 							}
 							else
@@ -286,6 +270,15 @@ namespace UE::RivermaxCore::Private
 			UE_LOG(LogRivermax, Log, TEXT("Rivermax Input stream has shutdown"));
 		}
 
+		if (bIsDynamicHeaderEnabled)
+		{
+			if (IRivermaxCoreModule* RivermaxModule = FModuleManager::GetModulePtr<IRivermaxCoreModule>(TEXT("RivermaxCore")))
+			{
+				RivermaxModule->GetRivermaxManager()->DisableDynamicHeaderSupport(Options.InterfaceAddress);
+			}
+			bIsDynamicHeaderEnabled = false;
+		}
+
 		DeallocateBuffers();
 	}
 
@@ -352,45 +345,6 @@ namespace UE::RivermaxCore::Private
 
 	}
 
-	bool FRivermaxInputStream::TranslateRTPHeader(const FRawRTPHeader& RawRTPHeader, FRTPHeader& OutRTPHeader)
-	{
-		OutRTPHeader.Timestamp = 0;
-		
-		if (RawRTPHeader.Version != 2)
-		{
-			return false;
-		}
-
-		// Pretty sure some data needs to be swapped but can't validate that until we have other hardware generating data
-		OutRTPHeader.SequencerNumber = (ByteSwap((uint16)RawRTPHeader.ExtendedSequenceNumber) << 16) | ByteSwap((uint16)RawRTPHeader.SequenceNumber);
-		OutRTPHeader.Timestamp = ByteSwap(RawRTPHeader.Timestamp);
-		OutRTPHeader.bIsMarkerBit = RawRTPHeader.MarkerBit;
-
-		OutRTPHeader.SyncSouceId = RawRTPHeader.SynchronizationSource;
-
-		OutRTPHeader.SRD1.Length = ByteSwap((uint16)RawRTPHeader.SRDLength1);
-		OutRTPHeader.SRD1.DataOffset = RawRTPHeader.GetSrd1Offset();
-		OutRTPHeader.SRD1.RowNumber = RawRTPHeader.GetSrd1RowNumber();
-		OutRTPHeader.SRD1.bIsFieldOne = RawRTPHeader.FieldIdentification1;
-		OutRTPHeader.SRD1.bHasContinuation = RawRTPHeader.ContinuationBit1;
-
-		if (OutRTPHeader.SRD1.bHasContinuation)
-		{
-			OutRTPHeader.SRD2.Length = ByteSwap((uint16)RawRTPHeader.SRDLength2);
-			OutRTPHeader.SRD2.DataOffset = RawRTPHeader.GetSrd2Offset();
-			OutRTPHeader.SRD2.RowNumber = RawRTPHeader.GetSrd2RowNumber();
-			OutRTPHeader.SRD2.bIsFieldOne = RawRTPHeader.FieldIdentification2;
-			OutRTPHeader.SRD2.bHasContinuation = RawRTPHeader.ContinuationBit2;
-
-			if (OutRTPHeader.SRD2.bHasContinuation == true)
-			{
-				UE_LOG(LogRivermax, Verbose, TEXT("Received SRD with more than 2 SRD which isn't supported."));
-			}
-		}
-
-		return true;
-	}
-
 	void FRivermaxInputStream::ParseChunks(const rmax_in_completion& Completion)
 	{
 		for (uint64 StrideIndex = 0; StrideIndex < Completion.chunk_size; ++StrideIndex)
@@ -413,13 +367,11 @@ namespace UE::RivermaxCore::Private
 
 			if (Completion.packet_info_arr[StrideIndex].data_size && RawHeaderPtr && DataPtr)
 			{
-				FRTPHeader RTPHeader;
-
 				// Get RTPHeader address from the raw net header
 				const FRawRTPHeader& RawRTPHeaderPtr = reinterpret_cast<const FRawRTPHeader&>(*GetRTPHeaderPointer(RawHeaderPtr));
-				const bool bIsValid = TranslateRTPHeader(RawRTPHeaderPtr, RTPHeader);
-				if (bIsValid)
+				if (RawRTPHeaderPtr.Version == 2)
 				{
+					FRTPHeader RTPHeader(RawRTPHeaderPtr);
 					// Add trace for the first packet of a frame to help visualize reception of a full frame in time
 					if (bIsFirstPacketReceived == false)
 					{
@@ -429,6 +381,8 @@ namespace UE::RivermaxCore::Private
 					}
 
 					StreamStats.BytesReceived += Completion.packet_info_arr[StrideIndex].data_size + Completion.packet_info_arr[StrideIndex].hdr_size;
+					
+					UpdateFrameTracking(RTPHeader);
 
 					switch (State)
 					{
@@ -472,9 +426,13 @@ namespace UE::RivermaxCore::Private
 		FRivermaxInputVideoFrameDescriptor Descriptor;
 		Descriptor.Timestamp = RTPHeader.Timestamp;
 		Descriptor.FrameNumber = TimestampToFrameNumber(RTPHeader.Timestamp, Options.FrameRate);
+		Descriptor.Width = StreamResolution.X;
+		Descriptor.Height = StreamResolution.Y;
+		Descriptor.PixelFormat = Options.PixelFormat;
+		const uint32 PixelCount = StreamResolution.X * StreamResolution.Y;
+		const uint32 FrameSize = PixelCount / FormatInfo.PixelGroupCoverage * FormatInfo.PixelGroupSize;
+		Descriptor.VideoBufferSize = FrameSize;
 		FRivermaxInputVideoFrameRequest Request;
-		const int32 Stride = Options.AlignedResolution.X / FormatInfo.PixelGroupCoverage * FormatInfo.PixelGroupSize;
-		Descriptor.VideoBufferSize = Options.Resolution.Y * Stride;
 		Listener->OnVideoFrameRequested(Descriptor, Request);
 
 		// Reset current frame to know when we have a valid one
@@ -514,9 +472,6 @@ namespace UE::RivermaxCore::Private
 		StreamData.SizeToWriteTwo = 0;
 		bIsFirstPacketReceived = false;
 
-		// New frame starting, reset tracked SRD
-		LastSRDLength.Reset();
-
 		return true;
 	}
 
@@ -833,15 +788,25 @@ namespace UE::RivermaxCore::Private
 			}
 		}
 
+		if (RTPHeader.SRD1.Length <= 0)
+		{
+			return;
+		}
+
 		uint32 DataOffset = 0;
-		StreamData.ReceivedSize += RTPHeader.SRD1.Length;
+		uint32 PayloadSize = RTPHeader.SRD1.Length;
+		if (RTPHeader.SRD1.bHasContinuation)
+		{
+			PayloadSize += RTPHeader.SRD2.Length;
+		}
+
 		if (bIsUsingGPUDirect)
 		{
 			// Initial case (start address)
 			if (StreamData.DeviceWritePointerOne == nullptr)
 			{
 				StreamData.DeviceWritePointerOne = DataPtr;
-				StreamData.SizeToWriteOne = RTPHeader.SRD1.Length;
+				StreamData.SizeToWriteOne = PayloadSize;
 			}
 			else
 			{
@@ -855,11 +820,11 @@ namespace UE::RivermaxCore::Private
 				// Case where we track memory in first buffer
 				if (StreamData.DeviceWritePointerTwo == nullptr)
 				{
-					StreamData.SizeToWriteOne += RTPHeader.SRD1.Length;
+					StreamData.SizeToWriteOne += PayloadSize;
 				}
 				else // Tracking memory in second buffer
 				{
-					StreamData.SizeToWriteTwo += RTPHeader.SRD1.Length;
+					StreamData.SizeToWriteTwo += PayloadSize;
 				}
 			}
 		}
@@ -874,18 +839,21 @@ namespace UE::RivermaxCore::Private
 				DataOffset += RTPHeader.SRD1.Length;
 				FMemory::Memcpy(&WriteBuffer[StreamData.WritingOffset], &DataPtr[DataOffset], RTPHeader.SRD2.Length);
 				StreamData.WritingOffset += RTPHeader.SRD2.Length;
-				StreamData.ReceivedSize += RTPHeader.SRD2.Length;
 			}
 		}
+
+		StreamData.ReceivedSize += PayloadSize;
 	}
 
 	void FRivermaxInputStream::ProcessLastSRD(const FRTPHeader& RTPHeader, const uint8* DataPtr)
 	{
 		FRivermaxInputVideoFrameDescriptor Descriptor;
-		Descriptor.Width = Options.Resolution.X;
-		Descriptor.Height = Options.Resolution.Y;
-		Descriptor.Stride = Options.AlignedResolution.X / FormatInfo.PixelGroupCoverage * FormatInfo.PixelGroupSize;
+		Descriptor.Width = StreamResolution.X;
+		Descriptor.Height = StreamResolution.Y;
+		Descriptor.PixelFormat = Options.PixelFormat;
 		Descriptor.Timestamp = RTPHeader.Timestamp;
+		const uint32 PixelCount = StreamResolution.X * StreamResolution.Y;
+		Descriptor.VideoBufferSize = PixelCount / FormatInfo.PixelGroupCoverage * FormatInfo.PixelGroupSize;
 		Descriptor.FrameNumber = TimestampToFrameNumber(RTPHeader.Timestamp, Options.FrameRate);
 
 		if (StreamData.ReceivedSize == StreamData.ExpectedSize)
@@ -897,13 +865,13 @@ namespace UE::RivermaxCore::Private
 
 			if (bIsUsingGPUDirect)
 			{
-				if (LastSRDLength.IsSet() && LastSRDLength.GetValue() != RTPHeader.SRD1.Length)
+				// For gpudirect, we need all payload sizes to be equal with the exception of the last one.
+				if (FrameDescriptionTracking.CommonPayloadSize <= 0)
 				{
-					UE_LOG(LogRivermax, Warning, TEXT("Unsupported variable SRD length detected while GPUDirect for input stream is used. Disable and reopen the stream. (Last: %d, New: %d)"), LastSRDLength.GetValue(), RTPHeader.SRD1.Length);
+					UE_LOG(LogRivermax, Warning, TEXT("Unsupported variable SRD length detected while GPUDirect for input stream is used. Disable and reopen the stream."));
 					Listener->OnStreamError();
 					bIsShuttingDown = true;
 				}
-				LastSRDLength = RTPHeader.SRD1.Length;
 
 				//Frame received entirely, time to copy it from rivermax gpu scratchpad to our own gpu memory
 				FCUDAModule& CudaModule = FModuleManager::GetModuleChecked<FCUDAModule>("CUDA");
@@ -912,8 +880,8 @@ namespace UE::RivermaxCore::Private
 				const CUdeviceptr DestinationGPUMemory = reinterpret_cast<CUdeviceptr>(StreamData.CurrentFrame);
 				const CUdeviceptr SourceGPUMemoryOne = reinterpret_cast<CUdeviceptr>(StreamData.DeviceWritePointerOne);
 
-				const uint32 NumSRDPartOne = RTPHeader.SRD1.Length > 0 ? StreamData.SizeToWriteOne / RTPHeader.SRD1.Length : 0;
-				const uint32 NumSRDPartTwo = RTPHeader.SRD1.Length > 0 ? StreamData.SizeToWriteTwo / RTPHeader.SRD1.Length : 0;
+				const uint32 NumSRDPartOne = StreamData.SizeToWriteOne / FrameDescriptionTracking.CommonPayloadSize;
+				const uint32 NumSRDPartTwo = StreamData.SizeToWriteTwo / FrameDescriptionTracking.CommonPayloadSize;
 
 				// Use cuda's 2d memcopy to do a source and destination stride difference memcopy
 				// We initialize rivermax stream with a payload size blindly since we don't know what the sender will use
@@ -927,17 +895,35 @@ namespace UE::RivermaxCore::Private
 				StrideDescription.dstMemoryType = CUmemorytype::CU_MEMORYTYPE_DEVICE;
 				StrideDescription.srcMemoryType = CUmemorytype::CU_MEMORYTYPE_DEVICE;
 				StrideDescription.srcPitch = ExpectedPayloadSize; //Source pitch is the expected payload used at init
-				StrideDescription.dstPitch = RTPHeader.SRD1.Length; //Destination pitch is the fixed SRD size we received
-				StrideDescription.WidthInBytes = RTPHeader.SRD1.Length; //Width in bytes is the amount to copy, the SRD size
+				StrideDescription.dstPitch = FrameDescriptionTracking.CommonPayloadSize; //Destination pitch is the fixed SRD size we received
+				StrideDescription.WidthInBytes = FrameDescriptionTracking.CommonPayloadSize; //Width in bytes is the amount to copy, the SRD size
 				StrideDescription.Height = NumSRDPartOne;
 				Result = CudaModule.DriverAPI()->cuMemcpy2DAsync(&StrideDescription, reinterpret_cast<CUstream>(GPUStream));
 
+				// Keep track where we are at in the source buffer
+				CUdeviceptr SourcePtr = SourceGPUMemoryOne + (ExpectedPayloadSize * NumSRDPartOne);
+
 				if (StreamData.DeviceWritePointerTwo != nullptr && StreamData.SizeToWriteTwo > 0)
 				{
 					StrideDescription.srcDevice = reinterpret_cast<CUdeviceptr>(StreamData.DeviceWritePointerTwo);
 					StrideDescription.dstDevice = reinterpret_cast<CUdeviceptr>(StreamData.CurrentFrame) + StreamData.SizeToWriteOne;
 					StrideDescription.Height = NumSRDPartTwo;
 					Result = CudaModule.DriverAPI()->cuMemcpy2DAsync(&StrideDescription, reinterpret_cast<CUstream>(GPUStream));
+
+					// Update source pointer since we used the second pointer tracker
+					SourcePtr = StrideDescription.srcDevice + (ExpectedPayloadSize * NumSRDPartTwo);
+				}
+
+				// Verify if we copied whole frame using the stride 2d copy or we're missing the last SRD of fractional size
+				const uint32 TotalMemoryCopied = (FrameDescriptionTracking.CommonPayloadSize * (NumSRDPartOne + NumSRDPartTwo));
+				if (TotalMemoryCopied < StreamData.ExpectedSize)
+				{
+					const uint32 SizeLeftToCopy = StreamData.ExpectedSize - TotalMemoryCopied;
+					if (ensure(SizeLeftToCopy == RTPHeader.SRD1.Length))
+					{
+						CUdeviceptr DestinationPtr = DestinationGPUMemory + TotalMemoryCopied;
+						Result = CudaModule.DriverAPI()->cuMemcpyDtoDAsync(DestinationPtr, SourcePtr, SizeLeftToCopy, reinterpret_cast<CUstream>(GPUStream));
+					}
 				}
 
 				if (Result != CUDA_SUCCESS)
@@ -990,6 +976,8 @@ namespace UE::RivermaxCore::Private
 				}
 			}
 
+			// Finished processing incoming frame, reset tracking data associated with it
+			FrameDescriptionTracking.ResetSingleFrameTracking();
 
 			// No need to provide the new frame and prepare the next one if we are shutting down
 			if (bIsShuttingDown == false)
@@ -1020,10 +1008,13 @@ namespace UE::RivermaxCore::Private
 	void FRivermaxInputStream::FrameReceptionState(const FRTPHeader& RTPHeader, const uint8* DataPtr)
 	{
 		FRivermaxInputVideoFrameDescriptor Descriptor;
-		Descriptor.Width = Options.Resolution.X;
-		Descriptor.Height = Options.Resolution.Y;
-		Descriptor.Stride = Options.AlignedResolution.X / FormatInfo.PixelGroupCoverage * FormatInfo.PixelGroupSize;
+		Descriptor.Width = StreamResolution.X;
+		Descriptor.Height = StreamResolution.Y;
+		Descriptor.PixelFormat = Options.PixelFormat;
 		Descriptor.Timestamp = RTPHeader.Timestamp;
+		const uint32 PixelCount = StreamResolution.X * StreamResolution.Y;
+		const uint32 FrameSize = PixelCount / FormatInfo.PixelGroupCoverage * FormatInfo.PixelGroupSize;
+		Descriptor.VideoBufferSize = FrameSize;
 		Descriptor.FrameNumber = TimestampToFrameNumber(RTPHeader.Timestamp, Options.FrameRate);
 
 		const uint64 LastSequenceNumberIncremented = StreamData.LastSequenceNumber + 1;
@@ -1095,5 +1086,29 @@ namespace UE::RivermaxCore::Private
 		WaitForMarkerState(RTPHeader);
 	}
 
+	void FRivermaxInputStream::UpdateFrameTracking(const FRTPHeader& NewRTPHeader)
+	{
+		FrameDescriptionTracking.EvaluateNewRTP(NewRTPHeader);
+		
+		if (!Options.bEnforceVideoFormat)
+		{
+			if (FrameDescriptionTracking.DetectedResolution.X > 0 && FrameDescriptionTracking.DetectedResolution.Y > 0)
+			{
+				if (FrameDescriptionTracking.DetectedResolution != StreamResolution)
+				{
+					StreamResolution = FrameDescriptionTracking.DetectedResolution;
+
+					FRivermaxInputVideoFormatChangedInfo FormatChangeInfo;
+					FormatChangeInfo.Width = StreamResolution.X;
+					FormatChangeInfo.Height = StreamResolution.Y;
+					FormatChangeInfo.PixelFormat = FrameDescriptionTracking.ExpectedSamplingType;
+					Listener->OnVideoFormatChanged(FormatChangeInfo);
+				}
+			}
+		}
+	}
+
 }
 
+	
+
diff --git a/Engine/Plugins/VirtualProduction/Rivermax/RivermaxCore/Source/RivermaxCore/Private/RivermaxInputStream.h b/Engine/Plugins/VirtualProduction/Rivermax/RivermaxCore/Source/RivermaxCore/Private/RivermaxInputStream.h
--- a/Engine/Plugins/VirtualProduction/Rivermax/RivermaxCore/Source/RivermaxCore/Private/RivermaxInputStream.h
+++ b/Engine/Plugins/VirtualProduction/Rivermax/RivermaxCore/Source/RivermaxCore/Private/RivermaxInputStream.h
@@ -27,7 +27,7 @@ namespace UE::RivermaxCore::Private
 		size_t PayloadSize = 0;
 		size_t HeaderSize = 0;
 		uint16 PayloadExpectedSize = 1500;
-		uint16 HeaderExpectedSize = 20; //for 2110
+		uint16 HeaderExpectedSize = 32;
 
 		rmax_in_memblock DataMemory;
 		rmax_in_memblock HeaderMemory;
@@ -62,6 +62,45 @@ namespace UE::RivermaxCore::Private
 		uint32 SizeToWriteTwo = 0;
 	};
 
+	/** Incoming frame description tracker. Used to detect variable SRD length across a frame for now but should be used to help auto detect incoming frame */
+	struct FFrameDescriptionTrackingData
+	{
+	public:
+
+		/** Resets data related to current frame being tracked */
+		void ResetSingleFrameTracking();
+
+		/** Look at new RTP and extract info to estimate frame description */
+		void EvaluateNewRTP(const FRTPHeader& NewHeader);
+
+		/** Whether a valid resolution was detected during the last frame */
+		bool HasDetectedValidResolution() const;
+
+	private:
+		
+		/** Updates current resolution detection using new RTP header. */
+		void UpdateResolutionDetection(const FRTPHeader& NewHeader);
+		
+		/** Updates payload sizes encountered throughout a frame */
+		void UpdatePayloadSizeTracking(const FRTPHeader& NewHeader);
+
+	public:
+		/** Used to track different payload sizes received. GPUDirect can't work if different payload sizes are received across the frame except the last one */
+		TMap<uint32, uint32> PayloadSizeReceived;
+
+		/** Payload size used across a frame. -1 if multi SRD except last one are received. */
+		int16 CommonPayloadSize = -1;
+
+		/** Last resolution detected */
+		FIntPoint DetectedResolution = FIntPoint::ZeroValue;
+
+		/** Sampling type we are expecting as setup by the user. We could also auto detect it but some sampling types would collide so left manual for now. */
+		ESamplingType ExpectedSamplingType = ESamplingType::RGB_10bit;
+
+		/** Used to log sampling type warning once until a good frame type is received */
+		bool bHasLoggedSamplingWarning = false;
+	};
+
 	class FRivermaxInputStream : public IRivermaxInputStream, public FRunnable
 	{
 	public:
@@ -131,6 +170,9 @@ namespace UE::RivermaxCore::Private
 		/** Gets the mapped gpu buffer for the incoming RHI buffer. Caching will occur if a new buffer is given */
 		void* GetMappedBuffer(const FBufferRHIRef& InBuffer);
 
+		/** Updates frame information tracking using RTP header */
+		void UpdateFrameTracking(const FRTPHeader& NewRTPHeader);
+
 	private:
 
 		/** Options used for this stream, such has resolution, frame rate etc... */
@@ -181,6 +223,9 @@ namespace UE::RivermaxCore::Private
 		/** Whether gpudirect is used for this stream. Even though listener requests it, doesn't mean it will/can be used */
 		bool bIsUsingGPUDirect = false;
 
+		/** Whether this stream has dynamic header data split enabled. Need to disable it on shutdown */
+		bool bIsDynamicHeaderEnabled = false;
+
 		/** GPU device index to use when mapping to cuda. Will be needed for mgpu support */
 		int32 GPUDeviceIndex = 0;
 
@@ -206,9 +251,6 @@ namespace UE::RivermaxCore::Private
 		/** Payload size used to initialize rivermax stream and know differences between received SRD and config */
 		uint32 ExpectedPayloadSize = 1500;
 
-		/** Used to track SRD length received and detect a change across a frame, which doesn't work with gpudirect */
-		TOptional<uint16> LastSRDLength = 0;
-
 		/** 
 		 * States used to track reception of a frame
 		 */
@@ -226,6 +268,12 @@ namespace UE::RivermaxCore::Private
 
 		/** Starts stream waiting for a marker to align ourselves with a full frame */
 		EReceptionState State = EReceptionState::WaitingForMarker;
+
+		/** Track incoming frame description to infer resolution, pixel type. */
+		FFrameDescriptionTrackingData FrameDescriptionTracking;
+
+		/** Resolution being used by this stream. Can either be auto detected or enforced by owner */
+		FIntPoint StreamResolution = FIntPoint::ZeroValue;
 	};
 }
 
diff --git a/Engine/Plugins/VirtualProduction/Rivermax/RivermaxCore/Source/RivermaxCore/Private/RivermaxManager.cpp b/Engine/Plugins/VirtualProduction/Rivermax/RivermaxCore/Source/RivermaxCore/Private/RivermaxManager.cpp
--- a/Engine/Plugins/VirtualProduction/Rivermax/RivermaxCore/Source/RivermaxCore/Private/RivermaxManager.cpp
+++ b/Engine/Plugins/VirtualProduction/Rivermax/RivermaxCore/Source/RivermaxCore/Private/RivermaxManager.cpp
@@ -448,10 +448,99 @@ namespace UE::RivermaxCore::Private
 	{
 		return bIsGPUDirectSupported && CVarRivermaxEnableGPUDirectOutputCapability.GetValueOnAnyThread();
 	}
-}
 
+	bool FRivermaxManager::EnableDynamicHeaderSupport(const FString& Interface)
+	{
+		if(uint32* Users = DynamicHeaderUsers.Find(Interface))
+		{
+			(*Users)++;
+			return true;
+		}
+
+		sockaddr_in RivermaxDevice;
+		FMemory::Memset(&RivermaxDevice, 0, sizeof(RivermaxDevice));
+		if (inet_pton(AF_INET, StringCast<ANSICHAR>(*Interface).Get(), &RivermaxDevice.sin_addr) != 1)
+		{
+			UE_LOG(LogRivermax, Warning, TEXT("Failed to use inet_pton for interface '%s'"), *Interface);
+			return false;
+		}
+
+		uint64_t CapabilityMask = RMAX_DEV_CAP_RTP_DYNAMIC_HDS;
+		rmax_device_caps_t Capabilities;
+		Capabilities.supported_caps = 0;
+
+		rmax_status_t Status = rmax_device_get_caps(RivermaxDevice.sin_addr, CapabilityMask, &Capabilities);
+		if (Status != RMAX_OK) 
+		{
+			UE_LOG(LogRivermax, Warning, TEXT("Failed to query capabilities for device IP '%s'"), *Interface);
+			return false;
+		}
+
+		const bool bIsDynamicHeaderSplitSupported = (Capabilities.supported_caps & RMAX_DEV_CAP_RTP_DYNAMIC_HDS) != 0;
+		if (!bIsDynamicHeaderSplitSupported) 
+		{
+			UE_LOG(LogRivermax, Warning, TEXT("RTP dynamic header data split is not supported for device IP '%s'"), *Interface)
+			return false;
+		}
+
+		rmax_device_config_t DeviceConfig;
+		FMemory::Memset(&DeviceConfig, 0, sizeof(DeviceConfig));
+		DeviceConfig.config_flags = RMAX_DEV_CONFIG_RTP_SMPTE_2110_20_DYNAMIC_HDS_CONFIG;
+		DeviceConfig.ip_address = RivermaxDevice.sin_addr;
+
+		Status = rmax_set_device_config(&DeviceConfig);
+		if (Status != RMAX_OK)
+		{
+			UE_LOG(LogRivermax, Warning, TEXT("Could not enable dynamic header split for device IP '%s'. Error: %d"), *Interface, Status);
+			return false;
+		}
+
+		UE_LOG(LogRivermax, Log, TEXT("Dynamic header split for device IP '%s' has been enabled."), *Interface);
+		DynamicHeaderUsers.FindOrAdd(Interface, 1);
+		return true;
+	}
+
+	void FRivermaxManager::DisableDynamicHeaderSupport(const FString& Interface)
+	{
+		if (uint32* Users = DynamicHeaderUsers.Find(Interface))
+		{
+			(*Users)--;
+			if((*Users) > 0)
+			{
+				return;
+			}
+
+			bool bCanUnsetDevice = true;
+			sockaddr_in DeviceToUse;
+			FMemory::Memset(&DeviceToUse, 0, sizeof(DeviceToUse));
+			if (inet_pton(AF_INET, StringCast<ANSICHAR>(*Interface).Get(), &DeviceToUse.sin_addr) != 1)
+			{
+				UE_LOG(LogRivermax, Warning, TEXT("Failed to use inet_pton for interface '%s'"), *Interface);
+				bCanUnsetDevice = false;
+			}
+
+			if (bCanUnsetDevice)
+			{
+				rmax_device_config_t DeviceConfig;
+				FMemory::Memset(&DeviceConfig, 0, sizeof(DeviceConfig));
+				DeviceConfig.config_flags = RMAX_DEV_CONFIG_RTP_SMPTE_2110_20_DYNAMIC_HDS_CONFIG;
+				DeviceConfig.ip_address = DeviceToUse.sin_addr;
 
+				const rmax_status_t Status = rmax_unset_device_config(&DeviceConfig);
+				if (Status != RMAX_OK)
+				{
+					UE_LOG(LogRivermax, Warning, TEXT("Could not disable dynamic header split for device IP '%s'. Error: %d"), *Interface, Status);
+				}
+				else
+				{
+					UE_LOG(LogRivermax, Log, TEXT("Dynamic header split for device IP '%s' has been disabled."), *Interface);
+				}
+			}
 
+			DynamicHeaderUsers.Remove(Interface);
+		}
+	}
 
+}
 
 
diff --git a/Engine/Plugins/VirtualProduction/Rivermax/RivermaxCore/Source/RivermaxCore/Private/RivermaxManager.h b/Engine/Plugins/VirtualProduction/Rivermax/RivermaxCore/Source/RivermaxCore/Private/RivermaxManager.h
--- a/Engine/Plugins/VirtualProduction/Rivermax/RivermaxCore/Source/RivermaxCore/Private/RivermaxManager.h
+++ b/Engine/Plugins/VirtualProduction/Rivermax/RivermaxCore/Source/RivermaxCore/Private/RivermaxManager.h
@@ -24,9 +24,11 @@ namespace UE::RivermaxCore::Private
 		virtual TConstArrayView<FRivermaxDeviceInfo> GetDevices() const override;
 		virtual bool GetMatchingDevice(const FString& InSourceIP, FString& OutDeviceIP) const override;
 		virtual bool IsValidIP(const FString& InSourceIP) const override;
-		virtual bool IsGPUDirectSupported() const;
-		virtual bool IsGPUDirectInputSupported() const;
-		virtual bool IsGPUDirectOutputSupported() const;
+		virtual bool IsGPUDirectSupported() const override;
+		virtual bool IsGPUDirectInputSupported() const override;
+		virtual bool IsGPUDirectOutputSupported() const override;
+		virtual bool EnableDynamicHeaderSupport(const FString& Interface) override;
+		virtual void DisableDynamicHeaderSupport(const FString& Interface) override;
 		//~ End IRivermaxManager interface
 
 	private:
@@ -45,7 +47,6 @@ namespace UE::RivermaxCore::Private
 
 		/** Verify if all prerequesites to go through Rivermax initialization have been met */
 		bool VerifyPrerequesites();
-
 	private:
 
 		/** True when manager has been initialized */
@@ -71,6 +72,9 @@ namespace UE::RivermaxCore::Private
 
 		/** Delegate triggered when initialization has completed */
 		FOnPostRivermaxManagerInit PostInitDelegate;
+
+		/** Streams using dynamic header size per users. When last one is removing itself, it will be disabled */
+		TMap<FString, uint32> DynamicHeaderUsers;
 	};
 }
 
diff --git a/Engine/Plugins/VirtualProduction/Rivermax/RivermaxCore/Source/RivermaxCore/Private/RivermaxOutputFrame.cpp b/Engine/Plugins/VirtualProduction/Rivermax/RivermaxCore/Source/RivermaxCore/Private/RivermaxOutputFrame.cpp
--- a/Engine/Plugins/VirtualProduction/Rivermax/RivermaxCore/Source/RivermaxCore/Private/RivermaxOutputFrame.cpp
+++ b/Engine/Plugins/VirtualProduction/Rivermax/RivermaxCore/Source/RivermaxCore/Private/RivermaxOutputFrame.cpp
@@ -30,6 +30,7 @@ namespace UE::RivermaxCore::Private
 		PacketCounter = 0;
 		LineNumber = 0;
 		SRDOffset = 0;
+		BytesSent = 0;
 	}
 
 	void FRivermaxOutputFrame::Reset()
diff --git a/Engine/Plugins/VirtualProduction/Rivermax/RivermaxCore/Source/RivermaxCore/Private/RivermaxOutputFrame.h b/Engine/Plugins/VirtualProduction/Rivermax/RivermaxCore/Source/RivermaxCore/Private/RivermaxOutputFrame.h
--- a/Engine/Plugins/VirtualProduction/Rivermax/RivermaxCore/Source/RivermaxCore/Private/RivermaxOutputFrame.h
+++ b/Engine/Plugins/VirtualProduction/Rivermax/RivermaxCore/Source/RivermaxCore/Private/RivermaxOutputFrame.h
@@ -52,6 +52,7 @@ namespace UE::RivermaxCore::Private
 		uint32 LineNumber = 0;
 		uint16 SRDOffset = 0;
 		uint32 ChunkNumber = 0;
+		uint32 BytesSent = 0;
 
 		/** Timestamp of this frame  */
 		double TimestampTicks = 0.0;
diff --git a/Engine/Plugins/VirtualProduction/Rivermax/RivermaxCore/Source/RivermaxCore/Private/RivermaxOutputStream.cpp b/Engine/Plugins/VirtualProduction/Rivermax/RivermaxCore/Source/RivermaxCore/Private/RivermaxOutputStream.cpp
--- a/Engine/Plugins/VirtualProduction/Rivermax/RivermaxCore/Source/RivermaxCore/Private/RivermaxOutputStream.cpp
+++ b/Engine/Plugins/VirtualProduction/Rivermax/RivermaxCore/Source/RivermaxCore/Private/RivermaxOutputStream.cpp
@@ -35,6 +35,22 @@ namespace UE::RivermaxCore::Private
 		TEXT("Interval at which to show stats in seconds"),
 		ECVF_Default);
 
+	static TAutoConsoleVariable<int32> CVarRivermaxOutputEnableMultiSRD(
+		TEXT("Rivermax.Output.EnableMultiSRD"), 1,
+		TEXT("When enabled, payload size will be constant across the frame except the last one.\n" 
+		     "If disabled, a payload size that fits the line will be used causing some resolution to not be supported."),
+		ECVF_Default);
+
+	static TAutoConsoleVariable<int32> CVarRivermaxOutputLinesPerChunk(
+		TEXT("Rivermax.Output.LinesPerChunk"), 4,
+		TEXT("Defines the number of lines to pack in a chunk. Higher number will increase latency"),
+		ECVF_Default);
+
+	static TAutoConsoleVariable<int32> CVarRivermaxOutputMaximizePacketSize(
+		TEXT("Rivermax.Output.MaximizePacketSize"), 1,
+		TEXT("Enables bigger packet sizes to maximize utilisation of potential UDP packet. If not enabled, packet size will be aligned with HD/4k sizes"),
+		ECVF_Default);
+
 	bool FindPayloadSize(const FRivermaxOutputStreamOptions& InOptions, uint32 InBytesPerLine, const FVideoFormatInfo& FormatInfo, uint16& OutPayloadSize)
 	{
 		using namespace UE::RivermaxCore::Private::Utils;
@@ -79,21 +95,190 @@ namespace UE::RivermaxCore::Private
 		return false;
 	}
 
-	uint32 FindChunksPerLine(const FRivermaxOutputStreamOptions& InOptions)
+	uint32 FindLinesPerChunk(const FRivermaxOutputStreamOptions& InOptions)
+	{
+		// More lines per chunks mean we will do more work prior to start sending a chunk. So, added 'latency' in terms of packet / parts of frame.
+		// Less lines per chunk mean that sender thread might starve.
+		// SDK sample uses 4 lines for UHD and 8 for HD. 
+		return CVarRivermaxOutputLinesPerChunk.GetValueOnAnyThread();
+	}
+
+	uint16 GetPayloadSize(ESamplingType SamplingType)
 	{
-		uint32 ChunksPerLine = 1; //Will need to revisit the impact of that
-		static constexpr uint32 MaxChunksPerLine = 10;
-		while (ChunksPerLine <= MaxChunksPerLine)
+		const FVideoFormatInfo FormatInfo = FStandardVideoFormat::GetVideoFormatInfo(SamplingType);
+		uint16 PayloadSize;
+		switch (SamplingType)
 		{
-			if (InOptions.Resolution.Y % ChunksPerLine == 0)
-			{
-				return ChunksPerLine;
-			}
 
-			++ChunksPerLine;
+		case ESamplingType::YUV444_10bit:
+		{
+			// Passthrough
+		}
+		case ESamplingType::RGB_10bit:
+		{
+			PayloadSize = 1200;
+			break;
+		}
+		case ESamplingType::YUV444_8bit:
+		{
+			// Passthrough
+		}
+		case ESamplingType::RGB_8bit:
+		{
+			PayloadSize = 1152;
+			break;
+		}
+		case ESamplingType::YUV444_12bit:
+		{
+			// Passthrough
+		}
+		case ESamplingType::RGB_12bit:
+		{
+			PayloadSize = 1152;
+			break;
+		}
+		case ESamplingType::YUV444_16bit:
+		{
+			// Passthrough
+		}
+		case ESamplingType::YUV444_16bitFloat:
+		{
+			// Passthrough
+		}
+		case ESamplingType::RGB_16bit:
+		{
+			// Passthrough
+		}
+		case ESamplingType::RGB_16bitFloat:
+		{
+			PayloadSize = 1152;
+			break;
+		}
+		case ESamplingType::YUV422_8bit:
+		{
+			PayloadSize = 1280;
+			break;
+		}
+		case ESamplingType::YUV422_10bit:
+		{
+			PayloadSize = 1200;
+			break;
+		}
+		case ESamplingType::YUV422_12bit:
+		{
+			PayloadSize = 1152;
+			break;
+		}
+		case ESamplingType::YUV422_16bit:
+		{
+			// Passthrough
+		}
+		case ESamplingType::YUV422_16bitFloat:
+		{
+			PayloadSize = 1280;
+			break;
+		}
+		default:
+		{
+			checkNoEntry();
+			PayloadSize = 1200;
+			break;
+		}
 		}
 
-		return 0;
+		ensure(PayloadSize % FormatInfo.PixelGroupSize == 0);
+		return PayloadSize;
+	}
+
+	/** 
+	 * Returns a payload closer to the max value we can have for standard UDP size 
+	 * RTPHeader can be bigger depending on configuration so we'll cap payload at 1400.
+	 */
+	uint16 GetMaximizedPayloadSize(ESamplingType SamplingType)
+	{
+		const FVideoFormatInfo FormatInfo = FStandardVideoFormat::GetVideoFormatInfo(SamplingType);
+		uint16 PayloadSize;
+		switch (SamplingType)
+		{
+		
+		case ESamplingType::YUV444_10bit:
+		{
+			// Passthrough
+		}
+		case ESamplingType::RGB_10bit:
+		{
+			PayloadSize = 1395;
+			break;
+		}
+		case ESamplingType::YUV444_8bit:
+		{
+			// Passthrough
+		}
+		case ESamplingType::RGB_8bit:
+		{
+			PayloadSize = 1398;	
+			break;
+		}
+		case ESamplingType::YUV444_12bit:
+		{
+			// Passthrough
+		}
+		case ESamplingType::RGB_12bit:
+		{
+			PayloadSize = 1395;
+			break;
+		}
+		case ESamplingType::YUV444_16bit:
+		{
+			// Passthrough
+		}
+		case ESamplingType::YUV444_16bitFloat:
+		{
+			// Passthrough
+		}
+		case ESamplingType::RGB_16bit:
+		{
+			// Passthrough
+		}
+		case ESamplingType::RGB_16bitFloat:
+		{
+			PayloadSize = 1398;
+			break;
+		}
+		case ESamplingType::YUV422_8bit:
+		{
+			PayloadSize = 1400;
+			break;
+		}
+		case ESamplingType::YUV422_10bit:
+		{
+			PayloadSize = 1400;
+			break;
+		}
+		case ESamplingType::YUV422_12bit:
+		{
+			PayloadSize = 1398;
+			break;
+		}
+		case ESamplingType::YUV422_16bit:
+		{
+			// Passthrough
+		}
+		case ESamplingType::YUV422_16bitFloat:
+		{
+			PayloadSize = 1400;
+			break;
+		}
+		default:
+		{
+			checkNoEntry();
+			PayloadSize = 1200;
+			break;
+		}
+		}
+
+		ensure(PayloadSize % FormatInfo.PixelGroupSize == 0);
+		return PayloadSize;
 	}
 
 
@@ -161,8 +346,8 @@ namespace UE::RivermaxCore::Private
 
 				rmax_buffer_attr BufferAttributes;
 				FMemory::Memset(&BufferAttributes, 0, sizeof(BufferAttributes));
-				BufferAttributes.chunk_size_in_strides = StreamMemory.ChunkSizeInStrides;
-				BufferAttributes.data_stride_size = StreamMemory.PayloadSize; //Stride between chunks. 
+				BufferAttributes.chunk_size_in_strides = StreamMemory.PacketsPerChunk;
+				BufferAttributes.data_stride_size = StreamMemory.PayloadSize; 
 				BufferAttributes.app_hdr_stride_size = StreamMemory.HeaderStrideSize;
 				BufferAttributes.mem_block_array = StreamMemory.MemoryBlocks.GetData();
 				BufferAttributes.mem_block_array_len = StreamMemory.MemoryBlocks.Num();
@@ -183,12 +368,11 @@ namespace UE::RivermaxCore::Private
 						StreamData.FrameFieldTimeIntervalNs = 1E9 / Options.FrameRate.AsDecimal();
 						InitializeStreamTimingSettings();
 
-						UE_LOG(LogRivermax, Display, TEXT("Output started to send %dx%d using %d payloads of size %d%s")
-						, Options.AlignedResolution.X
-						, Options.AlignedResolution.Y
-						, StreamMemory.PacketsInLine
-						, StreamMemory.PayloadSize
-						, bUseGPUDirect ? TEXT(" using GPUDirect") : TEXT(""));
+						UE_LOG(LogRivermax, Display, TEXT("Output stream started sending on stream %s:%d on interface %s%s")
+							, *Options.StreamAddress
+							, Options.Port
+							, *Options.InterfaceAddress
+							, bUseGPUDirect ? TEXT(" using GPUDirect") : TEXT(""));
 
 						bIsActive = true;
 						RivermaxThread.Reset(FRunnableThread::Create(this, TEXT("Rmax OutputStream Thread"), 128 * 1024, TPri_AboveNormal, FPlatformAffinity::GetPoolThreadMask()));
@@ -236,7 +420,7 @@ namespace UE::RivermaxCore::Private
 		if (TSharedPtr<FRivermaxOutputFrame> AvailableFrame = GetNextAvailableFrame(NewFrame.FrameIdentifier))
 		{
 			const int32 Stride = GetStride();
-			FMemory::Memcpy(AvailableFrame->VideoBuffer, NewFrame.VideoBuffer, NewFrame.Height * Stride);
+			FMemory::Memcpy(AvailableFrame->VideoBuffer, NewFrame.VideoBuffer, Options.AlignedResolution.Y * Stride);
 
 			AvailableFrame->bIsVideoBufferReady = true;
 
@@ -316,7 +500,7 @@ namespace UE::RivermaxCore::Private
 				//Update frame progress
 				if (bIsActive)
 				{
-					Stats.TotalStrides += StreamMemory.ChunkSizeInStrides;
+					Stats.TotalStrides += StreamMemory.PacketsPerChunk;
 					++CurrentFrame->ChunkNumber;
 				}
 
@@ -324,7 +508,6 @@ namespace UE::RivermaxCore::Private
 			
 
 			Stats.MemoryBlockSentCounter++;
-			Stats.TotalStrides += StreamMemory.ChunkSizeInStrides;
 			StreamData.bHasFrameFirstChunkBeenFetched = false;
 		}
 	}
@@ -358,50 +541,105 @@ namespace UE::RivermaxCore::Private
 		//We need to use the fullframe allocated size to compute the payload size.
 
 		const int32 BytesPerLine = GetStride();
-		const uint32 EffectiveBytesPerLine = BytesPerLine;
 
-		const bool bFoundPayload = FindPayloadSize(Options, BytesPerLine, FormatInfo, StreamMemory.PayloadSize);
-		if (bFoundPayload == false)
+		// Find out payload we want to use. Either we go the 'potential' multi SRD route or we keep the old way of finding a common payload
+		// with more restrictions on resolution supported. Kept in place to be able to fallback in case there are issues with the multiSRD one.
+		if (CVarRivermaxOutputEnableMultiSRD.GetValueOnAnyThread() >= 1)
 		{
-			UE_LOG(LogRivermax, Warning, TEXT("Could not find payload size for desired resolution %dx%d for desired pixel format"), Options.AlignedResolution.X, Options.AlignedResolution.Y);
-			return false;
+			if (CVarRivermaxOutputMaximizePacketSize.GetValueOnAnyThread() >= 1)
+			{
+				StreamMemory.PayloadSize = GetMaximizedPayloadSize(FormatInfo.Sampling);
+			}
+			else
+			{
+				StreamMemory.PayloadSize = GetPayloadSize(FormatInfo.Sampling);
+			}
 		}
+		else
+		{
+			const bool bFoundPayload = FindPayloadSize(Options, BytesPerLine, FormatInfo, StreamMemory.PayloadSize);
+			if (bFoundPayload == false)
+			{
+				UE_LOG(LogRivermax, Warning, TEXT("Could not find payload size for desired resolution %dx%d for desired pixel format"), Options.AlignedResolution.X, Options.AlignedResolution.Y);
+				return false;
+			}
+		}
+
+		// With payload size in hand, figure out how many packets we will need, how many chunks (group of packets) and configure descriptor arrays
+
+		const uint32 PixelCount = Options.AlignedResolution.X * Options.AlignedResolution.Y;
+		const uint64 FrameSize = PixelCount / FormatInfo.PixelGroupCoverage * FormatInfo.PixelGroupSize;
 
 		StreamMemory.PixelGroupPerPacket = StreamMemory.PayloadSize / FormatInfo.PixelGroupSize;
 		StreamMemory.PixelsPerPacket = StreamMemory.PixelGroupPerPacket * FormatInfo.PixelGroupCoverage;
-		StreamMemory.PacketsInLine = FMath::CeilToInt32((float)Options.AlignedResolution.X / StreamMemory.PixelsPerPacket);
 
-		StreamMemory.LinesInChunk = FindChunksPerLine(Options);
+		// We might need a smaller packet to complete the end of frame so ceil to the next value
+		StreamMemory.PacketsPerFrame = FMath::CeilToInt32((float)PixelCount / StreamMemory.PixelsPerPacket);
 
-		StreamMemory.ChunkSizeInStrides = StreamMemory.LinesInChunk * StreamMemory.PacketsInLine;
+		// Depending on resolution and payload size, last packet of a line might not be fully utilized but we need the remaining bytes so ceil to next value
+		StreamMemory.PacketsInLine = FMath::CeilToInt32((float)StreamMemory.PacketsPerFrame / Options.AlignedResolution.Y);
 
+		StreamMemory.LinesInChunk = FindLinesPerChunk(Options);
+		StreamMemory.PacketsPerChunk = StreamMemory.LinesInChunk * StreamMemory.PacketsInLine;
 		StreamMemory.FramesFieldPerMemoryBlock = 1;
-		StreamMemory.PacketsPerFrame =  StreamMemory.PacketsInLine * Options.Resolution.Y;
-		StreamMemory.PacketsPerMemoryBlock = StreamMemory.PacketsPerFrame * StreamMemory.FramesFieldPerMemoryBlock;
-		StreamMemory.ChunksPerFrameField = StreamMemory.PacketsPerFrame / StreamMemory.ChunkSizeInStrides;
+
+		// Chunk count won't necessarily align with the number of packets required. We need an integer amount of chunks to initialize our stream
+		// and calculate how many packets that represents. Rivermax will expect the payload/header array to be that size. It just means that
+		// we will mark the extra packets as 0 size.
+		StreamMemory.ChunksPerFrameField = FMath::CeilToInt32((float)StreamMemory.PacketsPerFrame / StreamMemory.PacketsPerChunk);
+		StreamMemory.PacketsPerMemoryBlock = StreamMemory.ChunksPerFrameField * StreamMemory.PacketsPerChunk * StreamMemory.FramesFieldPerMemoryBlock;
 		StreamMemory.ChunksPerMemoryBlock = StreamMemory.FramesFieldPerMemoryBlock * StreamMemory.ChunksPerFrameField;
 		StreamMemory.MemoryBlockCount = Options.NumberOfBuffers;
-		StreamMemory.StridesPerMemoryBlock = StreamMemory.ChunkSizeInStrides * StreamMemory.ChunksPerMemoryBlock;
 
 		// Setup arrays with the right sizes so we can give pointers to rivermax
-		StreamMemory.RTPHeaders.SetNumZeroed(StreamMemory.PacketsPerMemoryBlock);
+		StreamMemory.RTPHeaders.SetNumZeroed(StreamMemory.MemoryBlockCount);
 		StreamMemory.PayloadSizes.SetNumUninitialized(StreamMemory.PacketsPerMemoryBlock);
 		StreamMemory.HeaderSizes.SetNumUninitialized(StreamMemory.PacketsPerMemoryBlock);
-		StreamMemory.HeaderStrideSize = RTPHeaderSize;
+		StreamMemory.HeaderStrideSize = sizeof(FRawRTPHeader);
+
+		uint64 TotalSize = 0;
+		uint64 LineSize = 0;
 		for (int32 PayloadSizeIndex = 0; PayloadSizeIndex < StreamMemory.PayloadSizes.Num(); ++PayloadSizeIndex)
 		{
-			//Go through each chunk to have effective payload size to be sent (last one of each line could be smaller)
-			if ((PayloadSizeIndex + 1) % StreamMemory.PacketsInLine == 0)
+			uint32 HeaderSize = FRawRTPHeader::OneSRDSize;
+			uint32 ThisPayloadSize = StreamMemory.PayloadSize;
+			if (TotalSize < FrameSize)
 			{
-				const uint32 LeftOver = EffectiveBytesPerLine - ((StreamMemory.PacketsInLine - 1) * StreamMemory.PayloadSize);
-				StreamMemory.PayloadSizes[PayloadSizeIndex] = LeftOver;
-				check(LeftOver % FormatInfo.PixelGroupSize == 0);
+				if ((LineSize + StreamMemory.PayloadSize) == BytesPerLine)
+				{
+					LineSize = 0;
+				}
+				else if ((LineSize + StreamMemory.PayloadSize) > BytesPerLine)
+				{
+					HeaderSize = FRawRTPHeader::TwoSRDSize;
+					LineSize = StreamMemory.PayloadSize - (BytesPerLine - LineSize);
+					if (LineSize > BytesPerLine)
+					{
+						UE_LOG(LogRivermax, Warning, TEXT("Unsupported small resolution, %dx%d, needing more than 2 SRD to express"), Options.AlignedResolution.X, Options.AlignedResolution.Y);
+						return false;
+					}
+				}
+				else
+				{
+					// Keep track of line size offset to know when to use TwoSRDs
+					LineSize += StreamMemory.PayloadSize;
+				}
+
+				if ((TotalSize + StreamMemory.PayloadSize) > FrameSize)
+				{
+					HeaderSize = FRawRTPHeader::OneSRDSize;
+				}
 			}
 			else
 			{
-				StreamMemory.PayloadSizes[PayloadSizeIndex] = StreamMemory.PayloadSize;
+				// Extra header/payload required for the chunk alignment are set to 0. Nothing has to be sent out the wire.
+				HeaderSize = 0;
+				ThisPayloadSize = 0;
 			}
-			StreamMemory.HeaderSizes[PayloadSizeIndex] = StreamMemory.HeaderStrideSize;
+			
+			StreamMemory.HeaderSizes[PayloadSizeIndex] = HeaderSize;
+			StreamMemory.PayloadSizes[PayloadSizeIndex] = ThisPayloadSize;
+			TotalSize += ThisPayloadSize;
 		}
 		StreamMemory.MemoryBlocks.SetNumZeroed(StreamMemory.MemoryBlockCount);
 		for (uint32 BlockIndex = 0; BlockIndex < StreamMemory.MemoryBlockCount; ++BlockIndex)
@@ -411,7 +649,9 @@ namespace UE::RivermaxCore::Private
 			Block.app_hdr_size_arr = StreamMemory.HeaderSizes.GetData();
 			Block.data_size_arr = StreamMemory.PayloadSizes.GetData();
 			Block.data_ptr = AvailableFrames[BlockIndex]->VideoBuffer;
-			Block.app_hdr_ptr = &StreamMemory.RTPHeaders[BlockIndex];
+			
+			StreamMemory.RTPHeaders[BlockIndex].SetNumZeroed(StreamMemory.PacketsPerFrame);
+			Block.app_hdr_ptr = &StreamMemory.RTPHeaders[BlockIndex][0];
 		}
 
 		return true;
@@ -489,25 +729,16 @@ namespace UE::RivermaxCore::Private
 		return NextFrame;
 	}
 
-	void FRivermaxOutputStream::BuildRTPHeader(FOutputRTPHeader& OutHeader) const
+	void FRivermaxOutputStream::BuildRTPHeader(FRawRTPHeader& OutHeader) const
 	{
 		using namespace UE::RivermaxCore::Private::Utils;
 
-		// build RTP header - 12 bytes
-		/*
-		0                   1                   2                   3
-		0 1 2 3 4 5 6 7 8 9 0 1 2 3 4 5 6 7 8 9 0 1 2 3 4 5 6 7 8 9 0 1
-		+-+-+-+-+-+-+-+-+-+-+-+-+-+-+-+-+-+-+-+-+-+-+-+-+-+-+-+-+-+-+-+-+
-		| V |P|X|  CC   |M|     PT      |            SEQ                |
-		+-+-+-+-+-+-+-+-+-+-+-+-+-+-+-+-+-+-+-+-+-+-+-+-+-+-+-+-+-+-+-+-+
-		|                           timestamp                           |
-		+-+-+-+-+-+-+-+-+-+-+-+-+-+-+-+-+-+-+-+-+-+-+-+-+-+-+-+-+-+-+-+-+
-		|                           ssrc                                |
-		+-+-+-+-+-+-+-+-+-+-+-+-+-+-+-+-+-+-+-+-+-+-+-+-+-+-+-+-+-+-+-+-+*/
-		OutHeader.RawHeader[0] = 0x80;  // 10000000 - version2, no padding, no extension
-		OutHeader.RawHeader[1] = 96; //todo payload type from sdp file
-		OutHeader.RawHeader[2] = (StreamData.SequenceNumber >> 8) & 0xFF;  // sequence number
-		OutHeader.RawHeader[3] = (StreamData.SequenceNumber) & 0xFF;  // sequence number
+		OutHeader = {};
+		OutHeader.Version = 2;
+		OutHeader.PaddingBit = 0;
+		OutHeader.ExtensionBit = 0;
+		OutHeader.PayloadType = 96; //Payload type should probably be infered from SDP
+		OutHeader.SequenceNumber = ByteSwap((uint16)(StreamData.SequenceNumber & 0xFFFF));
 
 		// For now, in order to be able to use a framelocked input, we pipe frame number in the timestamp for a UE-UE interaction
 		// Follow up work to investigate adding this in RTP header
@@ -516,57 +747,65 @@ namespace UE::RivermaxCore::Private
 		{
 			InputTime = UE::RivermaxCore::GetAlignmentPointFromFrameNumber(CurrentFrame->FrameIdentifier, Options.FrameRate);
 		}
-		uint32 MediaTimestamp = GetTimestampFromTime(InputTime, MediaClockSampleRate);
-		
-		*(uint32*)&OutHeader.RawHeader[4] = ByteSwap(MediaTimestamp);
+
+		const uint32 MediaTimestamp = GetTimestampFromTime(InputTime, MediaClockSampleRate);
+		OutHeader.Timestamp = ByteSwap(MediaTimestamp);
 
 		//2110 specific header
-		*(uint32*)&OutHeader.RawHeader[8] = 0x0eb51dbd;  // simulated ssrc
+		OutHeader.SynchronizationSource = ByteSwap((uint32)0x0eb51dbd);  // Should Unreal has its own synch source ID
 
 		if (StreamType == ERivermaxStreamType::VIDEO_2110_20_STREAM)
 		{
-			//SRD means Sample Row Data
-			
-			// build SRD header - 8-14 bytes
-		   /* 0                   1                   2                   3
-			0 1 2 3 4 5 6 7 8 9 0 1 2 3 4 5 6 7 8 9 0 1 2 3 4 5 6 7 8 9 0 1
-			+-+-+-+-+-+-+-+-+-+-+-+-+-+-+-+-+-+-+-+-+-+-+-+-+-+-+-+-+-+-+-+-+
-			|    Extended Sequence Number   |           SRD Length          |
-			+-+-+-+-+-+-+-+-+-+-+-+-+-+-+-+-+-+-+-+-+-+-+-+-+-+-+-+-+-+-+-+-+
-			|F|     SRD Row Number          |C|         SRD Offset          |
-			+-+-+-+-+-+-+-+-+-+-+-+-+-+-+-+-+-+-+-+-+-+-+-+-+-+-+-+-+-+-+-+-+ */
-			OutHeader.RawHeader[12] = (StreamData.SequenceNumber >> 24) & 0xff;  // high 16 bit of seq Extended Sequence Number
-			OutHeader.RawHeader[13] = (StreamData.SequenceNumber >> 16) & 0xff;  // high 16 bit of seq Extended Sequence Number
-
-			const uint16 CurrentPayloadSize = StreamMemory.PayloadSizes[StreamData.SequenceNumber % StreamMemory.PacketsPerFrame];
-			*(uint16*)&OutHeader.RawHeader[14] = ByteSwap(CurrentPayloadSize);  // SRD Length
-		
-			uint16 number_of_rows = Options.Resolution.Y; //todo divide by 2 if interlaced
-			uint16 srd_row_number = (CurrentFrame->LineNumber % number_of_rows);
-			*(uint16*)&OutHeader.RawHeader[16] = ByteSwap(srd_row_number);
-			OutHeader.RawHeader[16] |= (0 << 7); //todo when fields are sent for interlace
+			if (CurrentFrame->PacketCounter + 1 == StreamMemory.PacketsPerFrame)
+			{
+				OutHeader.MarkerBit = 1; // last packet in frame (Marker)
+			}
 
-			// we never have continuation
+			OutHeader.ExtendedSequenceNumber = ByteSwap((uint16)((StreamData.SequenceNumber >> 16) & 0xFFFF));
 
-			// Write out current offset in pixels
-			*(uint16*)&OutHeader.RawHeader[18] = ByteSwap(CurrentFrame->SRDOffset);  // SRD Offset
+			// Verify if payload size exceeds line 
+			const uint32 CurrentPayloadSize = StreamMemory.PayloadSizes[CurrentFrame->PacketCounter];
 
-			// Update SRD offset in pixel for next round
-			uint16 PixelsPerPacket = (uint16)((StreamMemory.PayloadSize * FormatInfo.PixelGroupCoverage) / FormatInfo.PixelGroupSize);
-			CurrentFrame->SRDOffset = (CurrentFrame->SRDOffset + PixelsPerPacket) % (PixelsPerPacket * StreamMemory.PacketsInLine);
-		}
+			const uint32 LineSizeOffset = ((CurrentFrame->SRDOffset / FormatInfo.PixelGroupCoverage) * FormatInfo.PixelGroupSize);
+			const uint32 LineSize = ((Options.AlignedResolution.X / FormatInfo.PixelGroupCoverage) * FormatInfo.PixelGroupSize);
 
-		if (++CurrentFrame->PacketCounter == StreamMemory.PacketsPerFrame)
-		{
-			OutHeader.RawHeader[1] |= 0x80; // last packet in frame (Marker)
+			const uint16 SRD1Length = FMath::Min(LineSize - LineSizeOffset, CurrentPayloadSize);
+			const uint16 SRD1PixelCount = SRD1Length / FormatInfo.PixelGroupSize * FormatInfo.PixelGroupCoverage;
+			uint16 SRD2Length = SRD1Length < CurrentPayloadSize ? CurrentPayloadSize - SRD1Length : 0;
+			if (SRD2Length && CurrentFrame->LineNumber == ((uint32)Options.AlignedResolution.Y-1))
+			{
+				SRD2Length = 0;
+			}
+
+			OutHeader.SRD1Length = ByteSwap(SRD1Length);	
+			OutHeader.SetSrd1RowNumber(CurrentFrame->LineNumber); //todo divide by 2 if interlaced
+			OutHeader.FieldIdentification1 = 0; //todo when fields are sent for interlace
+			OutHeader.SetSrd1Offset(CurrentFrame->SRDOffset); 
+
+			CurrentFrame->SRDOffset += SRD1PixelCount;
+			if (CurrentFrame->SRDOffset >= Options.AlignedResolution.X)
+			{
+				CurrentFrame->SRDOffset = 0;
+				++CurrentFrame->LineNumber;
+			}
 
-			// ST2210-20: the timestamp SHOULD be the same for each packet of the frame/field.
-			const double ticks = (MediaClockSampleRate / (Options.FrameRate.AsDecimal()));
-			//if (set.video_type != VIDEO_TYPE::PROGRESSIVE) {
-			//	send_data.m_second_field = !send_data.m_second_field;
-			//	ticks /= 2;
-			//}
-			//CurrentFrame->TimestampTicks += ticks;
+			if (SRD2Length > 0)
+			{
+				OutHeader.SRD2Length = ByteSwap(SRD2Length);
+
+				OutHeader.ContinuationBit1 = 1;
+				OutHeader.FieldIdentification2 = 0;
+				OutHeader.SetSrd2RowNumber(CurrentFrame->LineNumber);
+				OutHeader.SetSrd2Offset(CurrentFrame->SRDOffset);
+
+				const uint16 SRD2PixelCount = SRD2Length / FormatInfo.PixelGroupSize * FormatInfo.PixelGroupCoverage;
+				CurrentFrame->SRDOffset += SRD2PixelCount;
+				if (CurrentFrame->SRDOffset >= Options.AlignedResolution.X)
+				{
+					CurrentFrame->SRDOffset = 0;
+					++CurrentFrame->LineNumber;
+				}
+			}
 		}
 	}
 
@@ -695,18 +934,15 @@ namespace UE::RivermaxCore::Private
 
 	void FRivermaxOutputStream::SetupRTPHeaders()
 	{
-		uint8* HeaderRawPtr = reinterpret_cast<uint8*>(CurrentFrame->HeaderPtr);
+		FRawRTPHeader* HeaderRawPtr = reinterpret_cast<FRawRTPHeader*>(CurrentFrame->HeaderPtr);
 		check(HeaderRawPtr);
-		for (uint32 StrideIndex = 0; StrideIndex < StreamMemory.ChunkSizeInStrides && CurrentFrame->PacketCounter < StreamMemory.PacketsPerFrame; ++StrideIndex)
+		for (uint32 StrideIndex = 0; StrideIndex < StreamMemory.PacketsPerChunk && CurrentFrame->PacketCounter < StreamMemory.PacketsPerFrame; ++StrideIndex)
 		{
-			uint8* NextHeaderRawPtr = HeaderRawPtr + (StrideIndex * StreamMemory.HeaderStrideSize);
-			BuildRTPHeader(*reinterpret_cast<FOutputRTPHeader*>(NextHeaderRawPtr));
-			//todo only for video
-			if (!((StrideIndex + 1) % StreamMemory.PacketsInLine))
-			{
-				CurrentFrame->LineNumber = (CurrentFrame->LineNumber + 1) % Options.Resolution.Y; //preparing line number for next iteration
-			}
+			BuildRTPHeader(*HeaderRawPtr);
 			++StreamData.SequenceNumber;
+			CurrentFrame->BytesSent += StreamMemory.PayloadSizes[CurrentFrame->PacketCounter];
+			++CurrentFrame->PacketCounter;
+			++HeaderRawPtr;
 		}
 	}
 
@@ -909,7 +1145,7 @@ namespace UE::RivermaxCore::Private
 			if (bIsProgressive)
 			{
 				RActive = (1080.0 / 1125.0);
-				if (Options.Resolution.Y >= FullHDHeight)
+				if (Options.AlignedResolution.Y >= FullHDHeight)
 				{
 					// As defined by SMPTE 2110-21 6.3.2
 					TRODefaultMultiplier = (43.0 / 1125.0);
@@ -921,13 +1157,13 @@ namespace UE::RivermaxCore::Private
 			}
 			else
 			{
-				if (Options.Resolution.Y >= FullHDHeight)
+				if (Options.AlignedResolution.Y >= FullHDHeight)
 				{
 					// As defined by SMPTE 2110-21 6.3.3
 					RActive = (1080.0 / 1125.0);
 					TRODefaultMultiplier = (22.0 / 1125.0);
 				}
-				else if (Options.Resolution.Y >= 576)
+				else if (Options.AlignedResolution.Y >= 576)
 				{
 					RActive = (576.0 / 625.0);
 					TRODefaultMultiplier = (26.0 / 625.0);
diff --git a/Engine/Plugins/VirtualProduction/Rivermax/RivermaxCore/Source/RivermaxCore/Private/RivermaxOutputStream.h b/Engine/Plugins/VirtualProduction/Rivermax/RivermaxCore/Source/RivermaxCore/Private/RivermaxOutputStream.h
--- a/Engine/Plugins/VirtualProduction/Rivermax/RivermaxCore/Source/RivermaxCore/Private/RivermaxOutputStream.h
+++ b/Engine/Plugins/VirtualProduction/Rivermax/RivermaxCore/Source/RivermaxCore/Private/RivermaxOutputStream.h
@@ -9,6 +9,7 @@
 #include "RivermaxHeader.h"
 #include "RivermaxOutputFrame.h"
 #include "RivermaxTypes.h"
+#include "RTPHeader.h"
 
 
 class FEvent;
@@ -18,12 +19,6 @@ namespace UE::RivermaxCore::Private
 {
 	using UE::RivermaxCore::FRivermaxOutputStreamOptions;
 	using UE::RivermaxCore::FRivermaxOutputVideoFrameInfo;
-	
-	// Todo make a proper RTP header struct
-	struct FOutputRTPHeader
-	{
-		uint8 RawHeader[20];
-	};
 
 	struct FRivermaxOutputStreamMemory
 	{
@@ -36,7 +31,7 @@ namespace UE::RivermaxCore::Private
 		uint32 LinesInChunk = 4;
 
 		uint32 PacketsInLine = 0;
-		uint32 ChunkSizeInStrides = 0;
+		uint32 PacketsPerChunk = 0;
 
 		uint32 FramesFieldPerMemoryBlock = 0;
 		uint32 PacketsPerFrame = 0;
@@ -44,12 +39,11 @@ namespace UE::RivermaxCore::Private
 		uint32 ChunksPerFrameField = 0;
 		uint32 ChunksPerMemoryBlock = 0;
 		uint32 MemoryBlockCount = 0; 
-		uint32 StridesPerMemoryBlock = 0;
 
 		TArray<rmax_mem_block> MemoryBlocks;
 		TArray<uint16_t> PayloadSizes; //Array describing stride payload size
 		TArray<uint16_t> HeaderSizes; //Array describing header payload size
-		TArray<FOutputRTPHeader> RTPHeaders;
+		TArray<TArray<FRawRTPHeader>> RTPHeaders;
 
 		rmax_buffer_attr BufferAttributes;
 	};
@@ -142,7 +136,7 @@ namespace UE::RivermaxCore::Private
 		TSharedPtr<FRivermaxOutputFrame> GetNextAvailableFrame(uint32 InFrameIdentifier);
 
 		/** Fills RTP and SRD header using current state */
-		void BuildRTPHeader(FOutputRTPHeader& OutHeader) const;
+		void BuildRTPHeader(FRawRTPHeader& OutHeader) const;
 
 		/** Destroys rivermax stream. Will wait until it's ready to be destroyed */
 		void DestroyStream();
diff --git a/Engine/Plugins/VirtualProduction/Rivermax/RivermaxCore/Source/RivermaxCore/Private/RivermaxUtils.cpp b/Engine/Plugins/VirtualProduction/Rivermax/RivermaxCore/Source/RivermaxCore/Private/RivermaxUtils.cpp
--- a/Engine/Plugins/VirtualProduction/Rivermax/RivermaxCore/Source/RivermaxCore/Private/RivermaxUtils.cpp
+++ b/Engine/Plugins/VirtualProduction/Rivermax/RivermaxCore/Source/RivermaxCore/Private/RivermaxUtils.cpp
@@ -111,8 +111,8 @@ namespace UE::RivermaxCore::Private::Utils
 		OutSDPDescription.Appendf("a=rtpmap:96 raw/90000\n");
 		OutSDPDescription.Appendf("a=fmtp: 96 sampling=%S; width=%d; height=%d; exactframerate=%S; depth=%S; TCS=SDR; colorimetry=BT709; PM=2110GPM; SSN=ST2110-20:2017; TP=2110TPN;\n"
 			, *PixelFormatToSamplingDesc(Options.PixelFormat)
-			, Options.Resolution.X
-			, Options.Resolution.Y
+			, Options.AlignedResolution.X
+			, Options.AlignedResolution.Y
 			, *FrameRateDescription
 			, *PixelFormatToBitDepth(Options.PixelFormat));
 	}
diff --git a/Engine/Plugins/VirtualProduction/Rivermax/RivermaxCore/Source/RivermaxCore/Private/RivermaxUtils.h b/Engine/Plugins/VirtualProduction/Rivermax/RivermaxCore/Source/RivermaxCore/Private/RivermaxUtils.h
--- a/Engine/Plugins/VirtualProduction/Rivermax/RivermaxCore/Source/RivermaxCore/Private/RivermaxUtils.h
+++ b/Engine/Plugins/VirtualProduction/Rivermax/RivermaxCore/Source/RivermaxCore/Private/RivermaxUtils.h
@@ -14,7 +14,7 @@ namespace UE::RivermaxCore
 namespace UE::RivermaxCore::Private::Utils
 {	
 	/** Various constants used for stream initialization */
-	static constexpr uint32 RTPHeaderSize = 20;
+
 	static constexpr uint32 FullHDHeight = 1080;
 	static constexpr uint32 FullHDWidth = 1920;
 
diff --git a/Engine/Plugins/VirtualProduction/Rivermax/RivermaxCore/Source/RivermaxCore/Public/IRivermaxInputStream.h b/Engine/Plugins/VirtualProduction/Rivermax/RivermaxCore/Source/RivermaxCore/Public/IRivermaxInputStream.h
--- a/Engine/Plugins/VirtualProduction/Rivermax/RivermaxCore/Source/RivermaxCore/Public/IRivermaxInputStream.h
+++ b/Engine/Plugins/VirtualProduction/Rivermax/RivermaxCore/Source/RivermaxCore/Public/IRivermaxInputStream.h
@@ -3,6 +3,7 @@
 #pragma once
 
 #include "CoreMinimal.h"
+#include "RivermaxFormats.h"
 
 class FRHIBuffer;
 
@@ -27,9 +28,6 @@ namespace UE::RivermaxCore
 		/** Width of the received frame */
 		uint32 Width = 0;
 
-		/** Size in bytes of a row */
-		uint32 Stride = 0;
-
 		/** Total size of the video frame */
 		uint32 VideoBufferSize = 0;
 
@@ -38,6 +36,21 @@ namespace UE::RivermaxCore
 
 		/** Frame number derived from timestamp and frame rate */
 		uint32 FrameNumber = 0;
+
+		/** Pixel format of the frame */
+		ESamplingType PixelFormat = ESamplingType::RGB_10bit;
+	};
+
+	struct RIVERMAXCORE_API FRivermaxInputVideoFormatChangedInfo
+	{
+		/** Detected height in pixel of the video stream */
+		uint32 Height = 0;
+
+		/** Detected width in pixel of the video stream */
+		uint32 Width = 0;
+
+		/** Detected sampling type of the video stream */
+		ESamplingType PixelFormat = ESamplingType::RGB_10bit;
 	};
 
 	struct RIVERMAXCORE_API FRivermaxInputVideoFrameRequest
@@ -71,6 +84,9 @@ namespace UE::RivermaxCore
 
 		/** Called when stream has encountered an error and has to stop */
 		virtual void OnStreamError() = 0;
+
+		/**  Called when stream has detected a change in the video format */
+		virtual void OnVideoFormatChanged(const FRivermaxInputVideoFormatChangedInfo& NewFormatInfo) = 0;
 	};
 
 	class RIVERMAXCORE_API IRivermaxInputStream
diff --git a/Engine/Plugins/VirtualProduction/Rivermax/RivermaxCore/Source/RivermaxCore/Public/IRivermaxManager.h b/Engine/Plugins/VirtualProduction/Rivermax/RivermaxCore/Source/RivermaxCore/Public/IRivermaxManager.h
--- a/Engine/Plugins/VirtualProduction/Rivermax/RivermaxCore/Source/RivermaxCore/Public/IRivermaxManager.h
+++ b/Engine/Plugins/VirtualProduction/Rivermax/RivermaxCore/Source/RivermaxCore/Public/IRivermaxManager.h
@@ -64,6 +64,12 @@ namespace UE::RivermaxCore
 
 		/** Whether gpudirect is enabled for output streams. */
 		virtual bool IsGPUDirectOutputSupported() const = 0;
+
+		/** Enables dynamic header support for a particular Rivermax device. */
+		virtual bool EnableDynamicHeaderSupport(const FString& Interface) = 0;
+
+		/** If it's the last user of that interface, we disable dynamic header */
+		virtual void DisableDynamicHeaderSupport(const FString& Interface) = 0;
 	};
 }
 
diff --git a/Engine/Plugins/VirtualProduction/Rivermax/RivermaxCore/Source/RivermaxCore/Public/RivermaxTypes.h b/Engine/Plugins/VirtualProduction/Rivermax/RivermaxCore/Source/RivermaxCore/Public/RivermaxTypes.h
--- a/Engine/Plugins/VirtualProduction/Rivermax/RivermaxCore/Source/RivermaxCore/Public/RivermaxTypes.h
+++ b/Engine/Plugins/VirtualProduction/Rivermax/RivermaxCore/Source/RivermaxCore/Public/RivermaxTypes.h
@@ -17,9 +17,6 @@ namespace UE::RivermaxCore
 
 	struct RIVERMAXCORE_API FRivermaxInputStreamOptions
 	{
-		/** Desired stream resolution */
-		FIntPoint Resolution = { 1920, 1080 };
-
 		/** Stream FrameRate */
 		FFrameRate FrameRate = { 24,1 };
 
@@ -37,9 +34,12 @@ namespace UE::RivermaxCore
 
 		/** Sample count to buffer. */
 		int32 NumberOfBuffers = 2;
+		
+		/** If true, don't use auto detected video format */
+		bool bEnforceVideoFormat = false;
 
-		/** Resolution aligning with pgroup of sampling type */
-		FIntPoint AlignedResolution = FIntPoint::ZeroValue;
+		/** Enforced resolution aligning with pgroup of sampling type */
+		FIntPoint EnforcedResolution = FIntPoint::ZeroValue;
 
 		/** Whether to leverage GPUDirect (Cuda) capability to transfer memory to NIC if available */
 		bool bUseGPUDirect = true;
diff --git a/Engine/Plugins/VirtualProduction/Rivermax/RivermaxCore/Source/RivermaxRendering/Private/RivermaxShaders.cpp b/Engine/Plugins/VirtualProduction/Rivermax/RivermaxCore/Source/RivermaxRendering/Private/RivermaxShaders.cpp
--- a/Engine/Plugins/VirtualProduction/Rivermax/RivermaxCore/Source/RivermaxRendering/Private/RivermaxShaders.cpp
+++ b/Engine/Plugins/VirtualProduction/Rivermax/RivermaxCore/Source/RivermaxRendering/Private/RivermaxShaders.cpp
@@ -13,14 +13,14 @@ namespace UE::RivermaxShaders
 
 	IMPLEMENT_GLOBAL_SHADER(FYUV8Bit422ToRGBACS, "/Plugin/RivermaxCore/Private/RivermaxShaders.usf", "YUV8Bit422ToRGBACS", SF_Compute);
 
-	FYUV8Bit422ToRGBACS::FParameters* FYUV8Bit422ToRGBACS::AllocateAndSetParameters(FRDGBuilder& GraphBuilder, FRDGBufferRef YUVBuffer, FRDGTextureRef OutputTexture, const FMatrix& ColorTransform, const FVector& YUVOffset, int32 BufferElementsPerRow, int32 BufferLineCount)
+	FYUV8Bit422ToRGBACS::FParameters* FYUV8Bit422ToRGBACS::AllocateAndSetParameters(FRDGBuilder& GraphBuilder, FRDGBufferRef YUVBuffer, FRDGTextureRef OutputTexture, const FMatrix& ColorTransform, const FVector& YUVOffset, int32 TotalElementCount, int32 BufferLineCount)
 	{
 		FYUV8Bit422ToRGBACS::FParameters* Parameters = GraphBuilder.AllocParameters<FYUV8Bit422ToRGBACS::FParameters>();
 
 		Parameters->InputYUV4228bitBuffer= GraphBuilder.CreateSRV(YUVBuffer);
 		Parameters->ColorTransform = (FMatrix44f)MediaShaders::CombineColorTransformAndOffset(ColorTransform, YUVOffset);
-		Parameters->HorizontalElementCount = BufferElementsPerRow;
-		Parameters->VerticalElementCount = BufferLineCount;
+		Parameters->TotalElementCount = TotalElementCount;
+		Parameters->OutputTextureSizeX = OutputTexture->Desc.Extent.X;
 		Parameters->OutTexture = GraphBuilder.CreateUAV(OutputTexture, ERDGUnorderedAccessViewFlags::None);
 
 		return Parameters;
@@ -31,29 +31,24 @@ namespace UE::RivermaxShaders
 
 	IMPLEMENT_GLOBAL_SHADER(FRGBToYUV8Bit422CS, "/Plugin/RivermaxCore/Private/RivermaxShaders.usf", "RGBToYUV8Bit422", SF_Compute);
 
-	FRGBToYUV8Bit422CS::FParameters* FRGBToYUV8Bit422CS::AllocateAndSetParameters(FRDGBuilder& GraphBuilder, FRDGTextureRef RGBATexture, const FIntPoint& SourceSize, const FIntRect& SourceViewRect, const FIntPoint& OutputSize, const FMatrix& ColorTransform, const FVector& YUVOffset, bool bDoLinearToSrgb, FRDGBufferRef OutputBuffer)
+	FRGBToYUV8Bit422CS::FParameters* FRGBToYUV8Bit422CS::AllocateAndSetParameters(FRDGBuilder& GraphBuilder, FRDGTextureRef SourceTexture, const FIntPoint& CaptureSize, const FIntRect& SourceViewRect, const FIntPoint& OutputSize, const FMatrix& ColorTransform, const FVector& YUVOffset, bool bDoLinearToSrgb, FRDGBufferRef OutputBuffer)
 	{
 		FRGBToYUV8Bit422CS::FParameters* Parameters = GraphBuilder.AllocParameters<FRGBToYUV8Bit422CS::FParameters>();
 
-		const uint32 InputTextureSizeX = SourceSize.X;
-		const uint32 InputTextureSizeY = SourceSize.Y;
-		Parameters->RGBToYUVConversion.InputTexture = RGBATexture;
-		Parameters->RGBToYUVConversion.InputSampler = TStaticSamplerState<SF_Point>::GetRHI();
+		Parameters->RGBToYUVConversion.InputTexture = SourceTexture;
+		Parameters->RGBToYUVConversion.InputSampler = TStaticSamplerState<SF_Point, AM_Border, AM_Border, AM_Border>::GetRHI();
 		Parameters->RGBToYUVConversion.ColorTransform = (FMatrix44f)MediaShaders::CombineColorTransformAndOffset(ColorTransform, YUVOffset);
 		Parameters->RGBToYUVConversion.DoLinearToSrgb = bDoLinearToSrgb;
-		Parameters->RGBToYUVConversion.OnePixelDeltaX = 1.0f / (float)InputTextureSizeX;
-		Parameters->OnePixelDeltaY = 1.0f / (float)InputTextureSizeY;
+		Parameters->RGBToYUVConversion.OnePixelDeltaX = 1.0f / (float)SourceTexture->Desc.Extent.X;
+		Parameters->OnePixelDeltaY = 1.0f / (float)SourceTexture->Desc.Extent.Y;
 
 		// Used to offset indices to sample input texture based on ThreadId of CS
+		Parameters->CapturedSizeX = CaptureSize.X;
 		Parameters->InputPixelOffsetX = SourceViewRect.Min.X;
 		Parameters->InputPixelOffsetY = SourceViewRect.Min.Y;
 
 		// Output size will be based on RGB to YUV422 encoding ( divide by 2 ) 
-		Parameters->HorizontalElementCount = OutputSize.X;
-		Parameters->VerticalElementCount = OutputSize.Y;
-
-		// Each thread will read many RGB pixels to convert to YUV
-		Parameters->InputTexturePixelsPerThread = InputTextureSizeX / OutputSize.X;
+		Parameters->TotalElementCount = OutputSize.X;
 
 		Parameters->OutYUV4228bitBuffer = GraphBuilder.CreateUAV(OutputBuffer);
 
@@ -65,14 +60,14 @@ namespace UE::RivermaxShaders
 
 	IMPLEMENT_GLOBAL_SHADER(FYUV10Bit422ToRGBACS, "/Plugin/RivermaxCore/Private/RivermaxShaders.usf", "YUV10Bit422ToRGBACS", SF_Compute);
 
-	FYUV10Bit422ToRGBACS::FParameters* FYUV10Bit422ToRGBACS::AllocateAndSetParameters(FRDGBuilder& GraphBuilder, FRDGBufferRef YUVBuffer, FRDGTextureRef OutputTexture, const FMatrix& ColorTransform, const FVector& YUVOffset, int32 BufferElementsPerRow, int32 BufferLineCount)
+	FYUV10Bit422ToRGBACS::FParameters* FYUV10Bit422ToRGBACS::AllocateAndSetParameters(FRDGBuilder& GraphBuilder, FRDGBufferRef YUVBuffer, FRDGTextureRef OutputTexture, const FMatrix& ColorTransform, const FVector& YUVOffset, int32 TotalElementCount, int32 BufferLineCount)
 	{
 		FYUV10Bit422ToRGBACS::FParameters* Parameters = GraphBuilder.AllocParameters<FYUV10Bit422ToRGBACS::FParameters>();
 
 		Parameters->InputYCbCrBuffer = GraphBuilder.CreateSRV(YUVBuffer);
 		Parameters->ColorTransform = (FMatrix44f)MediaShaders::CombineColorTransformAndOffset(ColorTransform, YUVOffset);
-		Parameters->HorizontalElementCount = BufferElementsPerRow;
-		Parameters->VerticalElementCount = BufferLineCount;
+		Parameters->TotalElementCount = TotalElementCount;
+		Parameters->OutputTextureSizeX = OutputTexture->Desc.Extent.X;
 		Parameters->OutTexture = GraphBuilder.CreateUAV(OutputTexture, ERDGUnorderedAccessViewFlags::None);
 
 		return Parameters;
@@ -83,29 +78,24 @@ namespace UE::RivermaxShaders
 
 	IMPLEMENT_GLOBAL_SHADER(FRGBToYUV10Bit422LittleEndianCS, "/Plugin/RivermaxCore/Private/RivermaxShaders.usf", "RGBToYUV10Bit422", SF_Compute);
 
-	FRGBToYUV10Bit422LittleEndianCS::FParameters* FRGBToYUV10Bit422LittleEndianCS::AllocateAndSetParameters(FRDGBuilder& GraphBuilder, FRDGTextureRef RGBATexture, const FIntPoint& SourceSize, const FIntRect& SourceViewRect, const FIntPoint& OutputSize, const FMatrix& ColorTransform, const FVector& YUVOffset, bool bDoLinearToSrgb, FRDGBufferRef OutputBuffer)
+	FRGBToYUV10Bit422LittleEndianCS::FParameters* FRGBToYUV10Bit422LittleEndianCS::AllocateAndSetParameters(FRDGBuilder& GraphBuilder, FRDGTextureRef SourceTexture, const FIntPoint& CaptureSize, const FIntRect& SourceViewRect, const FIntPoint& OutputSize, const FMatrix& ColorTransform, const FVector& YUVOffset, bool bDoLinearToSrgb, FRDGBufferRef OutputBuffer)
 	{
 		FRGBToYUV10Bit422LittleEndianCS::FParameters* Parameters = GraphBuilder.AllocParameters<FRGBToYUV10Bit422LittleEndianCS::FParameters>();
 
-		const uint32 InputTextureSizeX = SourceSize.X;
-		const uint32 InputTextureSizeY = SourceSize.Y;
-		Parameters->RGBToYUVConversion.InputTexture = RGBATexture;
-		Parameters->RGBToYUVConversion.InputSampler = TStaticSamplerState<SF_Point>::GetRHI();
+		Parameters->RGBToYUVConversion.InputTexture = SourceTexture;
+		Parameters->RGBToYUVConversion.InputSampler = TStaticSamplerState<SF_Point, AM_Border, AM_Border, AM_Border>::GetRHI();
 		Parameters->RGBToYUVConversion.ColorTransform = (FMatrix44f)MediaShaders::CombineColorTransformAndOffset(ColorTransform, YUVOffset);
 		Parameters->RGBToYUVConversion.DoLinearToSrgb = bDoLinearToSrgb;
-		Parameters->RGBToYUVConversion.OnePixelDeltaX = 1.0f / (float)InputTextureSizeX;
-		Parameters->OnePixelDeltaY = 1.0f / (float)InputTextureSizeY;
-
+		Parameters->RGBToYUVConversion.OnePixelDeltaX = 1.0f / (float)SourceTexture->Desc.Extent.X;
+		Parameters->OnePixelDeltaY = 1.0f / (float)SourceTexture->Desc.Extent.Y;
+		
 		// Used to offset indices to sample input texture based on ThreadId of CS
+		Parameters->CapturedSizeX = CaptureSize.X;
 		Parameters->InputPixelOffsetX = SourceViewRect.Min.X;
 		Parameters->InputPixelOffsetY = SourceViewRect.Min.Y;
 
 		// Output size will be based on RGB to YUV encoding ( divide by 2 ) and pixels encoded per element ( 4 pixels )
-		Parameters->HorizontalElementCount = OutputSize.X;
-		Parameters->VerticalElementCount = OutputSize.Y;
-
-		// Each thread will read many RGB pixels to convert to YUV
-		Parameters->InputTexturePixelsPerThread = InputTextureSizeX / OutputSize.X;
+		Parameters->TotalElementCount = OutputSize.X;
 
 		Parameters->OutYCbCrBuffer = GraphBuilder.CreateUAV(OutputBuffer);
 
@@ -117,13 +107,13 @@ namespace UE::RivermaxShaders
 
 	IMPLEMENT_GLOBAL_SHADER(FRGB8BitToRGBA8CS, "/Plugin/RivermaxCore/Private/RivermaxShaders.usf", "RGB8BitToRGBA8CS", SF_Compute);
 
-	FRGB8BitToRGBA8CS::FParameters* FRGB8BitToRGBA8CS::AllocateAndSetParameters(FRDGBuilder& GraphBuilder, FRDGBufferRef RGBBuffer, FRDGTextureRef OutputTexture, int32 BufferElementsPerRow, int32 BufferLineCount)
+	FRGB8BitToRGBA8CS::FParameters* FRGB8BitToRGBA8CS::AllocateAndSetParameters(FRDGBuilder& GraphBuilder, FRDGBufferRef RGBBuffer, FRDGTextureRef OutputTexture, int32 TotalElementCount, int32 BufferLineCount)
 	{
 		FRGB8BitToRGBA8CS::FParameters* Parameters = GraphBuilder.AllocParameters<FRGB8BitToRGBA8CS::FParameters>();
 
 		Parameters->InputRGB8Buffer = GraphBuilder.CreateSRV(RGBBuffer);
-		Parameters->HorizontalElementCount = BufferElementsPerRow;
-		Parameters->VerticalElementCount = BufferLineCount;
+		Parameters->TotalElementCount = TotalElementCount;
+		Parameters->OutputTextureSizeX = OutputTexture->Desc.Extent.X;
 		Parameters->OutTexture = GraphBuilder.CreateUAV(OutputTexture, ERDGUnorderedAccessViewFlags::None);
 
 		return Parameters;
@@ -134,26 +124,21 @@ namespace UE::RivermaxShaders
 
 	IMPLEMENT_GLOBAL_SHADER(FRGBToRGB8BitCS, "/Plugin/RivermaxCore/Private/RivermaxShaders.usf", "RGBToRGB8BitCS", SF_Compute);
 
-	FRGBToRGB8BitCS::FParameters* FRGBToRGB8BitCS::AllocateAndSetParameters(FRDGBuilder& GraphBuilder, FRDGTextureRef RGBATexture, const FIntPoint& SourceSize, const FIntRect& SourceViewRect, const FIntPoint& OutputSize, FRDGBufferRef OutputBuffer)
+	FRGBToRGB8BitCS::FParameters* FRGBToRGB8BitCS::AllocateAndSetParameters(FRDGBuilder& GraphBuilder, FRDGTextureRef SourceTexture, const FIntPoint& CaptureSize, const FIntRect& SourceViewRect, const FIntPoint& OutputSize, FRDGBufferRef OutputBuffer)
 	{
 		FRGBToRGB8BitCS::FParameters* Parameters = GraphBuilder.AllocParameters<FRGBToRGB8BitCS::FParameters>();
 
-		const uint32 InputTextureSizeX = SourceSize.X;
-		const uint32 InputTextureSizeY = SourceSize.Y;
-		Parameters->InputTexture = RGBATexture;
-		Parameters->InputSampler = TStaticSamplerState<SF_Point>::GetRHI();
-		Parameters->OnePixelDeltaX = 1.0f / (float)InputTextureSizeX;
-		Parameters->OnePixelDeltaY = 1.0f / (float)InputTextureSizeY;
+		Parameters->InputTexture = SourceTexture;
+		Parameters->InputSampler = TStaticSamplerState<SF_Point, AM_Border, AM_Border, AM_Border>::GetRHI();
+		Parameters->OnePixelDeltaX = 1.0f / (float)SourceTexture->Desc.Extent.X;
+		Parameters->OnePixelDeltaY = 1.0f / (float)SourceTexture->Desc.Extent.Y;
 
 		// Used to offset indices to sample input texture based on ThreadId of CS
+		Parameters->CapturedSizeX = CaptureSize.X;
 		Parameters->InputPixelOffsetX = SourceViewRect.Min.X;
 		Parameters->InputPixelOffsetY = SourceViewRect.Min.Y;
 
-		Parameters->HorizontalElementCount = OutputSize.X;
-		Parameters->VerticalElementCount = OutputSize.Y;
-
-		// Each thread will read many RGBA pixels to pack into RGB buffer
-		Parameters->InputTexturePixelsPerThread = InputTextureSizeX / OutputSize.X;
+		Parameters->TotalElementCount = OutputSize.X;
 
 		Parameters->OutRGB8Buffer = GraphBuilder.CreateUAV(OutputBuffer);
 
@@ -166,13 +151,13 @@ namespace UE::RivermaxShaders
 	IMPLEMENT_GLOBAL_SHADER(FRGB10BitToRGBA10CS, "/Plugin/RivermaxCore/Private/RivermaxShaders.usf", "RGB10BitToRGBACS", SF_Compute);
 
 
-	FRGB10BitToRGBA10CS::FParameters* FRGB10BitToRGBA10CS::AllocateAndSetParameters(FRDGBuilder& GraphBuilder, FRDGBufferRef RGBBuffer, FRDGTextureRef OutputTexture, int32 BufferElementsPerRow, int32 BufferLineCount)
+	FRGB10BitToRGBA10CS::FParameters* FRGB10BitToRGBA10CS::AllocateAndSetParameters(FRDGBuilder& GraphBuilder, FRDGBufferRef RGBBuffer, FRDGTextureRef OutputTexture, int32 TotalElementCount, int32 BufferLineCount)
 	{
 		FRGB10BitToRGBA10CS::FParameters* Parameters = GraphBuilder.AllocParameters<FRGB10BitToRGBA10CS::FParameters>();
 
 		Parameters->InputBuffer = GraphBuilder.CreateSRV(RGBBuffer);
-		Parameters->HorizontalElementCount = BufferElementsPerRow;
-		Parameters->VerticalElementCount = BufferLineCount;
+		Parameters->TotalElementCount = TotalElementCount;
+		Parameters->OutputTextureSizeX = OutputTexture->Desc.Extent.X;
 		Parameters->OutTexture = GraphBuilder.CreateUAV(OutputTexture, ERDGUnorderedAccessViewFlags::None);
 
 		return Parameters;
@@ -184,26 +169,22 @@ namespace UE::RivermaxShaders
 	IMPLEMENT_GLOBAL_SHADER(FRGBToRGB10BitCS, "/Plugin/RivermaxCore/Private/RivermaxShaders.usf", "RGBToRGB10BitCS", SF_Compute);
 
 
-	FRGBToRGB10BitCS::FParameters* FRGBToRGB10BitCS::AllocateAndSetParameters(FRDGBuilder& GraphBuilder, FRDGTextureRef RGBATexture, const FIntPoint& SourceSize, const FIntRect& SourceViewRect, const FIntPoint& OutputSize, FRDGBufferRef OutputBuffer)
+	FRGBToRGB10BitCS::FParameters* FRGBToRGB10BitCS::AllocateAndSetParameters(FRDGBuilder& GraphBuilder, FRDGTextureRef SourceTexture, const FIntPoint& CaptureSize, const FIntRect& SourceViewRect, const FIntPoint& OutputSize, FRDGBufferRef OutputBuffer)
 	{
 		FRGBToRGB10BitCS::FParameters* Parameters = GraphBuilder.AllocParameters<FRGBToRGB10BitCS::FParameters>();
 
-		const uint32 InputTextureSizeX = SourceSize.X;
-		const uint32 InputTextureSizeY = SourceSize.Y;
-		Parameters->InputTexture = RGBATexture;
-		Parameters->InputSampler = TStaticSamplerState<SF_Point>::GetRHI();
-		Parameters->OnePixelDeltaX = 1.0f / (float)InputTextureSizeX;
-		Parameters->OnePixelDeltaY = 1.0f / (float)InputTextureSizeY;
+		Parameters->InputTexture = SourceTexture;
+
+		Parameters->InputSampler = TStaticSamplerState<SF_Point, AM_Border, AM_Border, AM_Border>::GetRHI();
+		Parameters->OnePixelDeltaX = 1.0f / (float)SourceTexture->Desc.Extent.X;
+		Parameters->OnePixelDeltaY = 1.0f / (float)SourceTexture->Desc.Extent.Y;
 
 		// Used to offset indices to sample input texture based on ThreadId of CS
+		Parameters->CapturedSizeX = CaptureSize.X;
 		Parameters->InputPixelOffsetX = SourceViewRect.Min.X;
 		Parameters->InputPixelOffsetY = SourceViewRect.Min.Y;
 
-		Parameters->HorizontalElementCount = OutputSize.X;
-		Parameters->VerticalElementCount = OutputSize.Y;
-
-		// Each thread will read many RGBA pixels to pack into RGB buffer
-		Parameters->InputTexturePixelsPerThread = InputTextureSizeX / OutputSize.X;
+		Parameters->TotalElementCount = OutputSize.X;
 
 		Parameters->OutRGB10Buffer = GraphBuilder.CreateUAV(OutputBuffer);
 
@@ -216,13 +197,13 @@ namespace UE::RivermaxShaders
 	IMPLEMENT_GLOBAL_SHADER(FRGB12BitToRGBA12CS, "/Plugin/RivermaxCore/Private/RivermaxShaders.usf", "RGB12BitToRGBACS", SF_Compute);
 
 
-	FRGB12BitToRGBA12CS::FParameters* FRGB12BitToRGBA12CS::AllocateAndSetParameters(FRDGBuilder& GraphBuilder, FRDGBufferRef RGBBuffer, FRDGTextureRef OutputTexture, int32 BufferElementsPerRow, int32 BufferLineCount)
+	FRGB12BitToRGBA12CS::FParameters* FRGB12BitToRGBA12CS::AllocateAndSetParameters(FRDGBuilder& GraphBuilder, FRDGBufferRef RGBBuffer, FRDGTextureRef OutputTexture, int32 TotalElementCount, int32 BufferLineCount)
 	{
 		FRGB12BitToRGBA12CS::FParameters* Parameters = GraphBuilder.AllocParameters<FRGB12BitToRGBA12CS::FParameters>();
 
 		Parameters->InputRGB12Buffer = GraphBuilder.CreateSRV(RGBBuffer);
-		Parameters->HorizontalElementCount = BufferElementsPerRow;
-		Parameters->VerticalElementCount = BufferLineCount;
+		Parameters->TotalElementCount = TotalElementCount;
+		Parameters->OutputTextureSizeX = OutputTexture->Desc.Extent.X;
 		Parameters->OutTexture = GraphBuilder.CreateUAV(OutputTexture, ERDGUnorderedAccessViewFlags::None);
 
 		return Parameters;
@@ -234,26 +215,21 @@ namespace UE::RivermaxShaders
 	IMPLEMENT_GLOBAL_SHADER(FRGBToRGB12BitCS, "/Plugin/RivermaxCore/Private/RivermaxShaders.usf", "RGBToRGB12BitCS", SF_Compute);
 
 
-	FRGBToRGB12BitCS::FParameters* FRGBToRGB12BitCS::AllocateAndSetParameters(FRDGBuilder& GraphBuilder, FRDGTextureRef RGBATexture, const FIntPoint& SourceSize, const FIntRect& SourceViewRect, const FIntPoint& OutputSize, FRDGBufferRef OutputBuffer)
+	FRGBToRGB12BitCS::FParameters* FRGBToRGB12BitCS::AllocateAndSetParameters(FRDGBuilder& GraphBuilder, FRDGTextureRef SourceTexture, const FIntPoint& CaptureSize, const FIntRect& SourceViewRect, const FIntPoint& OutputSize, FRDGBufferRef OutputBuffer)
 	{
 		FRGBToRGB12BitCS::FParameters* Parameters = GraphBuilder.AllocParameters<FRGBToRGB12BitCS::FParameters>();
 
-		const uint32 InputTextureSizeX = SourceSize.X;
-		const uint32 InputTextureSizeY = SourceSize.Y;
-		Parameters->InputTexture = RGBATexture;
-		Parameters->InputSampler = TStaticSamplerState<SF_Point>::GetRHI();
-		Parameters->OnePixelDeltaX = 1.0f / (float)InputTextureSizeX;
-		Parameters->OnePixelDeltaY = 1.0f / (float)InputTextureSizeY;
+		Parameters->InputTexture = SourceTexture;
+		Parameters->InputSampler = TStaticSamplerState<SF_Point, AM_Border, AM_Border, AM_Border>::GetRHI();
+		Parameters->OnePixelDeltaX = 1.0f / (float)SourceTexture->Desc.Extent.X;
+		Parameters->OnePixelDeltaY = 1.0f / (float)SourceTexture->Desc.Extent.Y;
 
 		// Used to offset indices to sample input texture based on ThreadId of CS
+		Parameters->CapturedSizeX = CaptureSize.X;
 		Parameters->InputPixelOffsetX = SourceViewRect.Min.X;
 		Parameters->InputPixelOffsetY = SourceViewRect.Min.Y;
 
-		Parameters->HorizontalElementCount = OutputSize.X;
-		Parameters->VerticalElementCount = OutputSize.Y;
-
-		// Each thread will read many RGBA pixels to pack into RGB buffer
-		Parameters->InputTexturePixelsPerThread = InputTextureSizeX / OutputSize.X;
+		Parameters->TotalElementCount = OutputSize.X;
 
 		Parameters->OutRGB12Buffer = GraphBuilder.CreateUAV(OutputBuffer);
 
@@ -266,26 +242,21 @@ namespace UE::RivermaxShaders
 	IMPLEMENT_GLOBAL_SHADER(FRGBToRGB16fCS, "/Plugin/RivermaxCore/Private/RivermaxShaders.usf", "RGBToRGB16fBitCS", SF_Compute);
 
 
-	FRGBToRGB16fCS::FParameters* FRGBToRGB16fCS::AllocateAndSetParameters(FRDGBuilder& GraphBuilder, FRDGTextureRef RGBATexture, const FIntPoint& SourceSize, const FIntRect& SourceViewRect, const FIntPoint& OutputSize, FRDGBufferRef OutputBuffer)
+	FRGBToRGB16fCS::FParameters* FRGBToRGB16fCS::AllocateAndSetParameters(FRDGBuilder& GraphBuilder, FRDGTextureRef SourceTexture, const FIntPoint& CaptureSize, const FIntRect& SourceViewRect, const FIntPoint& OutputSize, FRDGBufferRef OutputBuffer)
 	{
 		FRGBToRGB16fCS::FParameters* Parameters = GraphBuilder.AllocParameters<FRGBToRGB16fCS::FParameters>();
 
-		const uint32 InputTextureSizeX = SourceSize.X;
-		const uint32 InputTextureSizeY = SourceSize.Y;
-		Parameters->InputTexture = RGBATexture;
-		Parameters->InputSampler = TStaticSamplerState<SF_Point>::GetRHI();
-		Parameters->OnePixelDeltaX = 1.0f / (float)InputTextureSizeX;
-		Parameters->OnePixelDeltaY = 1.0f / (float)InputTextureSizeY;
+		Parameters->InputTexture = SourceTexture;
+		Parameters->InputSampler = TStaticSamplerState<SF_Point, AM_Border, AM_Border, AM_Border>::GetRHI();
+		Parameters->OnePixelDeltaX = 1.0f / (float)SourceTexture->Desc.Extent.X;
+		Parameters->OnePixelDeltaY = 1.0f / (float)SourceTexture->Desc.Extent.Y;
 
 		// Used to offset indices to sample input texture based on ThreadId of CS
+		Parameters->CapturedSizeX = CaptureSize.X;
 		Parameters->InputPixelOffsetX = SourceViewRect.Min.X;
 		Parameters->InputPixelOffsetY = SourceViewRect.Min.Y;
 
-		Parameters->HorizontalElementCount = OutputSize.X;
-		Parameters->VerticalElementCount = OutputSize.Y;
-
-		// Each thread will read many RGBA pixels to pack into RGB buffer
-		Parameters->InputTexturePixelsPerThread = InputTextureSizeX / OutputSize.X;
+		Parameters->TotalElementCount = OutputSize.X;
 
 		Parameters->OutRGB16fBuffer = GraphBuilder.CreateUAV(OutputBuffer);
 
@@ -298,13 +269,13 @@ namespace UE::RivermaxShaders
 	IMPLEMENT_GLOBAL_SHADER(FRGB16fBitToRGBA16fCS, "/Plugin/RivermaxCore/Private/RivermaxShaders.usf", "RGB16fBitToRGBACS", SF_Compute);
 
 
-	FRGB16fBitToRGBA16fCS::FParameters* FRGB16fBitToRGBA16fCS::AllocateAndSetParameters(FRDGBuilder& GraphBuilder, FRDGBufferRef RGBBuffer, FRDGTextureRef OutputTexture, int32 BufferElementsPerRow, int32 BufferLineCount)
+	FRGB16fBitToRGBA16fCS::FParameters* FRGB16fBitToRGBA16fCS::AllocateAndSetParameters(FRDGBuilder& GraphBuilder, FRDGBufferRef RGBBuffer, FRDGTextureRef OutputTexture, int32 TotalElementCount, int32 BufferLineCount)
 	{
 		FRGB16fBitToRGBA16fCS::FParameters* Parameters = GraphBuilder.AllocParameters<FRGB16fBitToRGBA16fCS::FParameters>();
 
 		Parameters->InputRGB16fBuffer = GraphBuilder.CreateSRV(RGBBuffer);
-		Parameters->HorizontalElementCount = BufferElementsPerRow;
-		Parameters->VerticalElementCount = BufferLineCount;
+		Parameters->TotalElementCount = TotalElementCount;
+		Parameters->OutputTextureSizeX = OutputTexture->Desc.Extent.X;
 		Parameters->OutTexture = GraphBuilder.CreateUAV(OutputTexture, ERDGUnorderedAccessViewFlags::None);
 
 		return Parameters;
diff --git a/Engine/Plugins/VirtualProduction/Rivermax/RivermaxCore/Source/RivermaxRendering/Public/RivermaxShaders.h b/Engine/Plugins/VirtualProduction/Rivermax/RivermaxCore/Source/RivermaxRendering/Public/RivermaxShaders.h
--- a/Engine/Plugins/VirtualProduction/Rivermax/RivermaxCore/Source/RivermaxRendering/Public/RivermaxShaders.h
+++ b/Engine/Plugins/VirtualProduction/Rivermax/RivermaxCore/Source/RivermaxRendering/Public/RivermaxShaders.h
@@ -36,8 +36,8 @@ namespace UE::RivermaxShaders
 		BEGIN_SHADER_PARAMETER_STRUCT(FParameters, )
 			SHADER_PARAMETER_RDG_BUFFER_SRV(StructuredBuffer<FYUV8Bit422Buffer>, InputYUV4228bitBuffer)
 			SHADER_PARAMETER(FMatrix44f, ColorTransform)
-			SHADER_PARAMETER(uint32, HorizontalElementCount)
-			SHADER_PARAMETER(uint32, VerticalElementCount)
+			SHADER_PARAMETER(uint32, TotalElementCount)
+			SHADER_PARAMETER(uint32, OutputTextureSizeX)
 			SHADER_PARAMETER_RDG_TEXTURE_UAV(RWTexture2D, OutTexture)
 		END_SHADER_PARAMETER_STRUCT()
 
@@ -50,7 +50,7 @@ namespace UE::RivermaxShaders
 		}
 
 		/** Allocates and setup shader parameter in the incoming graph builder */
-		RIVERMAXRENDERING_API FYUV8Bit422ToRGBACS::FParameters* AllocateAndSetParameters(FRDGBuilder& GraphBuilder, FRDGBufferRef YUVBuffer, FRDGTextureRef OutputTexture, const FMatrix& ColorTransform, const FVector& YUVOffset, int32 BufferElementsPerRow, int32 BufferLineCount);
+		RIVERMAXRENDERING_API FYUV8Bit422ToRGBACS::FParameters* AllocateAndSetParameters(FRDGBuilder& GraphBuilder, FRDGBufferRef YUVBuffer, FRDGTextureRef OutputTexture, const FMatrix& ColorTransform, const FVector& YUVOffset, int32 TotalElementCount, int32 BufferLineCount);
 	};
 
 	/**
@@ -74,9 +74,8 @@ namespace UE::RivermaxShaders
 			SHADER_PARAMETER(float, OnePixelDeltaY)
 			SHADER_PARAMETER(float, InputPixelOffsetX)
 			SHADER_PARAMETER(float, InputPixelOffsetY)
-			SHADER_PARAMETER(uint32, HorizontalElementCount)
-			SHADER_PARAMETER(uint32, VerticalElementCount)
-			SHADER_PARAMETER(uint32, InputTexturePixelsPerThread)
+			SHADER_PARAMETER(uint32, CapturedSizeX)
+			SHADER_PARAMETER(uint32, TotalElementCount)
 			SHADER_PARAMETER_RDG_BUFFER_UAV(RWStructuredBuffer<FYUV8Bit422Buffer>, OutYUV4228bitBuffer)
 		END_SHADER_PARAMETER_STRUCT()
 
@@ -89,7 +88,7 @@ namespace UE::RivermaxShaders
 		}
 
 		/** Allocates and setup shader parameter in the incoming graph builder */
-		RIVERMAXRENDERING_API FRGBToYUV8Bit422CS::FParameters* AllocateAndSetParameters(FRDGBuilder& GraphBuilder, FRDGTextureRef RGBATexture, const FIntPoint& SourceSize, const FIntRect& SourceViewRect, const FIntPoint& OutputSize, const FMatrix& ColorTransform, const FVector& YUVOffset, bool bDoLinearToSrgb, FRDGBufferRef OutputBuffer);
+		RIVERMAXRENDERING_API FRGBToYUV8Bit422CS::FParameters* AllocateAndSetParameters(FRDGBuilder& GraphBuilder, FRDGTextureRef SourceTexture, const FIntPoint& CaptureSize, const FIntRect& SourceViewRect, const FIntPoint& OutputSize, const FMatrix& ColorTransform, const FVector& YUVOffset, bool bDoLinearToSrgb, FRDGBufferRef OutputBuffer);
 	};
 
 	/**
@@ -118,8 +117,8 @@ namespace UE::RivermaxShaders
 		BEGIN_SHADER_PARAMETER_STRUCT(FParameters, )
 			SHADER_PARAMETER_RDG_BUFFER_SRV(StructuredBuffer<FYUV10Bit422LEBuffer>, InputYCbCrBuffer)
 			SHADER_PARAMETER(FMatrix44f, ColorTransform)	
-			SHADER_PARAMETER(uint32, HorizontalElementCount)
-			SHADER_PARAMETER(uint32, VerticalElementCount)
+			SHADER_PARAMETER(uint32, TotalElementCount)
+			SHADER_PARAMETER(uint32, OutputTextureSizeX)
 			SHADER_PARAMETER_RDG_TEXTURE_UAV(RWTexture2D, OutTexture)
 		END_SHADER_PARAMETER_STRUCT()
 
@@ -132,7 +131,7 @@ namespace UE::RivermaxShaders
 		}
 
 		/** Allocates and setup shader parameter in the incoming graph builder */
-		RIVERMAXRENDERING_API FYUV10Bit422ToRGBACS::FParameters* AllocateAndSetParameters(FRDGBuilder& GraphBuilder, FRDGBufferRef YUVBuffer, FRDGTextureRef OutputTexture, const FMatrix& ColorTransform, const FVector& YUVOffset, int32 BufferElementsPerRow, int32 BufferLineCount);
+		RIVERMAXRENDERING_API FYUV10Bit422ToRGBACS::FParameters* AllocateAndSetParameters(FRDGBuilder& GraphBuilder, FRDGBufferRef YUVBuffer, FRDGTextureRef OutputTexture, const FMatrix& ColorTransform, const FVector& YUVOffset, int32 TotalElementCount, int32 BufferLineCount);
 	};
 
 	/**
@@ -162,9 +161,8 @@ namespace UE::RivermaxShaders
 			SHADER_PARAMETER(float, OnePixelDeltaY)
 			SHADER_PARAMETER(float, InputPixelOffsetX)
 			SHADER_PARAMETER(float, InputPixelOffsetY)
-			SHADER_PARAMETER(uint32, HorizontalElementCount)
-			SHADER_PARAMETER(uint32, VerticalElementCount)
-			SHADER_PARAMETER(uint32, InputTexturePixelsPerThread)
+			SHADER_PARAMETER(uint32, CapturedSizeX)
+			SHADER_PARAMETER(uint32, TotalElementCount)
 			SHADER_PARAMETER_RDG_BUFFER_UAV(RWStructuredBuffer<FYUV10Bit422LEBuffer>, OutYCbCrBuffer)
 		END_SHADER_PARAMETER_STRUCT()
 
@@ -177,7 +175,7 @@ namespace UE::RivermaxShaders
 		}
 
 		/** Allocates and setup shader parameter in the incoming graph builder */
-		RIVERMAXRENDERING_API FRGBToYUV10Bit422LittleEndianCS::FParameters* AllocateAndSetParameters(FRDGBuilder& GraphBuilder, FRDGTextureRef RGBATexture, const FIntPoint& SourceSize, const FIntRect& SourceViewRect, const FIntPoint& OutputSize, const FMatrix& ColorTransform, const FVector& YUVOffset, bool bDoLinearToSrgb, FRDGBufferRef OutputBuffer);
+		RIVERMAXRENDERING_API FRGBToYUV10Bit422LittleEndianCS::FParameters* AllocateAndSetParameters(FRDGBuilder& GraphBuilder, FRDGTextureRef SourceTexture, const FIntPoint& CaptureSize, const FIntRect& SourceViewRect, const FIntPoint& OutputSize, const FMatrix& ColorTransform, const FVector& YUVOffset, bool bDoLinearToSrgb, FRDGBufferRef OutputBuffer);
 	};
 
 	/**
@@ -204,8 +202,8 @@ namespace UE::RivermaxShaders
 
 		BEGIN_SHADER_PARAMETER_STRUCT(FParameters, )
 			SHADER_PARAMETER_RDG_BUFFER_SRV(StructuredBuffer<FRGB8BitBuffer>, InputRGB8Buffer)
-			SHADER_PARAMETER(uint32, HorizontalElementCount)
-			SHADER_PARAMETER(uint32, VerticalElementCount)
+			SHADER_PARAMETER(uint32, TotalElementCount)
+			SHADER_PARAMETER(uint32, OutputTextureSizeX)
 			SHADER_PARAMETER_RDG_TEXTURE_UAV(RWTexture2D, OutTexture)
 		END_SHADER_PARAMETER_STRUCT()
 
@@ -218,7 +216,7 @@ namespace UE::RivermaxShaders
 		}
 
 		/** Allocates and setup shader parameter in the incoming graph builder */
-		RIVERMAXRENDERING_API FRGB8BitToRGBA8CS::FParameters* AllocateAndSetParameters(FRDGBuilder& GraphBuilder, FRDGBufferRef RGBBuffer, FRDGTextureRef OutputTexture, int32 BufferElementsPerRow, int32 BufferLineCount);
+		RIVERMAXRENDERING_API FRGB8BitToRGBA8CS::FParameters* AllocateAndSetParameters(FRDGBuilder& GraphBuilder, FRDGBufferRef RGBBuffer, FRDGTextureRef OutputTexture, int32 TotalElementCount, int32 BufferLineCount);
 	};
 
 	/**
@@ -244,13 +242,12 @@ namespace UE::RivermaxShaders
 		BEGIN_SHADER_PARAMETER_STRUCT(FParameters, )
 			SHADER_PARAMETER_RDG_TEXTURE(Texture2D, InputTexture)
 			SHADER_PARAMETER_SAMPLER(SamplerState, InputSampler)
+			SHADER_PARAMETER(uint32, CapturedSizeX)
 			SHADER_PARAMETER(float, OnePixelDeltaX)
 			SHADER_PARAMETER(float, OnePixelDeltaY)
 			SHADER_PARAMETER(float, InputPixelOffsetX)
 			SHADER_PARAMETER(float, InputPixelOffsetY)
-			SHADER_PARAMETER(uint32, HorizontalElementCount)
-			SHADER_PARAMETER(uint32, VerticalElementCount)
-			SHADER_PARAMETER(uint32, InputTexturePixelsPerThread)
+			SHADER_PARAMETER(uint32, TotalElementCount)
 			SHADER_PARAMETER_RDG_BUFFER_UAV(RWStructuredBuffer<FRGB8BitBuffer>, OutRGB8Buffer)
 		END_SHADER_PARAMETER_STRUCT()
 
@@ -263,7 +260,7 @@ namespace UE::RivermaxShaders
 		}
 
 		/** Allocates and setup shader parameter in the incoming graph builder */
-		RIVERMAXRENDERING_API FRGBToRGB8BitCS::FParameters* AllocateAndSetParameters(FRDGBuilder& GraphBuilder, FRDGTextureRef RGBATexture, const FIntPoint& SourceSize, const FIntRect& SourceViewRect, const FIntPoint& OutputSize, FRDGBufferRef OutputBuffer);
+		RIVERMAXRENDERING_API FRGBToRGB8BitCS::FParameters* AllocateAndSetParameters(FRDGBuilder& GraphBuilder, FRDGTextureRef SourceTexture, const FIntPoint& CaptureSize, const FIntRect& SourceViewRect, const FIntPoint& OutputSize, FRDGBufferRef OutputBuffer);
 	};
 
 	/**
@@ -288,13 +285,12 @@ namespace UE::RivermaxShaders
 		BEGIN_SHADER_PARAMETER_STRUCT(FParameters, )
 			SHADER_PARAMETER_RDG_TEXTURE(Texture2D, InputTexture)
 			SHADER_PARAMETER_SAMPLER(SamplerState, InputSampler)
+			SHADER_PARAMETER(uint32, CapturedSizeX)
 			SHADER_PARAMETER(float, OnePixelDeltaX)
 			SHADER_PARAMETER(float, OnePixelDeltaY)
 			SHADER_PARAMETER(float, InputPixelOffsetX)
 			SHADER_PARAMETER(float, InputPixelOffsetY)
-			SHADER_PARAMETER(uint32, HorizontalElementCount)
-			SHADER_PARAMETER(uint32, VerticalElementCount)
-			SHADER_PARAMETER(uint32, InputTexturePixelsPerThread)
+			SHADER_PARAMETER(uint32, TotalElementCount)
 			SHADER_PARAMETER_RDG_BUFFER_UAV(RWStructuredBuffer<FRGB10BitBuffer>, OutRGB10Buffer)
 		END_SHADER_PARAMETER_STRUCT()
 
@@ -307,7 +303,7 @@ namespace UE::RivermaxShaders
 		}
 
 		/** Allocates and setup shader parameter in the incoming graph builder */
-		RIVERMAXRENDERING_API FRGBToRGB10BitCS::FParameters* AllocateAndSetParameters(FRDGBuilder& GraphBuilder, FRDGTextureRef RGBATexture, const FIntPoint& SourceSize, const FIntRect& SourceViewRect, const FIntPoint& OutputSize, FRDGBufferRef OutputBuffer);
+		RIVERMAXRENDERING_API FRGBToRGB10BitCS::FParameters* AllocateAndSetParameters(FRDGBuilder& GraphBuilder, FRDGTextureRef SourceTexture, const FIntPoint& CaptureSize, const FIntRect& SourceViewRect, const FIntPoint& OutputSize, FRDGBufferRef OutputBuffer);
 	};
 
 	/**
@@ -333,8 +329,8 @@ namespace UE::RivermaxShaders
 
 		BEGIN_SHADER_PARAMETER_STRUCT(FParameters, )
 			SHADER_PARAMETER_RDG_BUFFER_SRV(StructuredBuffer<FRGB10BitBuffer>, InputBuffer)
-			SHADER_PARAMETER(uint32, HorizontalElementCount)
-			SHADER_PARAMETER(uint32, VerticalElementCount)
+			SHADER_PARAMETER(uint32, TotalElementCount)
+			SHADER_PARAMETER(uint32, OutputTextureSizeX)
 			SHADER_PARAMETER_RDG_TEXTURE_UAV(RWTexture2D, OutTexture)
 		END_SHADER_PARAMETER_STRUCT()
 
@@ -347,7 +343,7 @@ namespace UE::RivermaxShaders
 		}
 
 		/** Allocates and setup shader parameter in the incoming graph builder */
-		RIVERMAXRENDERING_API FRGB10BitToRGBA10CS::FParameters* AllocateAndSetParameters(FRDGBuilder& GraphBuilder, FRDGBufferRef RGBBuffer, FRDGTextureRef OutputTexture, int32 BufferElementsPerRow, int32 BufferLineCount);
+		RIVERMAXRENDERING_API FRGB10BitToRGBA10CS::FParameters* AllocateAndSetParameters(FRDGBuilder& GraphBuilder, FRDGBufferRef RGBBuffer, FRDGTextureRef OutputTexture, int32 TotalElementCount, int32 BufferLineCount);
 	};
 
 	/**
@@ -372,13 +368,12 @@ namespace UE::RivermaxShaders
 		BEGIN_SHADER_PARAMETER_STRUCT(FParameters, )
 			SHADER_PARAMETER_RDG_TEXTURE(Texture2D, InputTexture)
 			SHADER_PARAMETER_SAMPLER(SamplerState, InputSampler)
+			SHADER_PARAMETER(uint32, CapturedSizeX)
 			SHADER_PARAMETER(float, OnePixelDeltaX)
 			SHADER_PARAMETER(float, OnePixelDeltaY)
 			SHADER_PARAMETER(float, InputPixelOffsetX)
 			SHADER_PARAMETER(float, InputPixelOffsetY)
-			SHADER_PARAMETER(uint32, HorizontalElementCount)
-			SHADER_PARAMETER(uint32, VerticalElementCount)
-			SHADER_PARAMETER(uint32, InputTexturePixelsPerThread)
+			SHADER_PARAMETER(uint32, TotalElementCount)
 			SHADER_PARAMETER_RDG_BUFFER_UAV(RWStructuredBuffer<FRGB12BitBuffer>, OutRGB12Buffer)
 		END_SHADER_PARAMETER_STRUCT()
 
@@ -391,7 +386,7 @@ namespace UE::RivermaxShaders
 		}
 
 		/** Allocates and setup shader parameter in the incoming graph builder */
-		RIVERMAXRENDERING_API FRGBToRGB12BitCS::FParameters* AllocateAndSetParameters(FRDGBuilder& GraphBuilder, FRDGTextureRef RGBATexture, const FIntPoint& SourceSize, const FIntRect& SourceViewRect, const FIntPoint& OutputSize, FRDGBufferRef OutputBuffer);
+		RIVERMAXRENDERING_API FRGBToRGB12BitCS::FParameters* AllocateAndSetParameters(FRDGBuilder& GraphBuilder, FRDGTextureRef SourceTexture, const FIntPoint& CaptureSize, const FIntRect& SourceViewRect, const FIntPoint& OutputSize, FRDGBufferRef OutputBuffer);
 	};
 
 	/**
@@ -417,8 +412,8 @@ namespace UE::RivermaxShaders
 
 		BEGIN_SHADER_PARAMETER_STRUCT(FParameters, )
 			SHADER_PARAMETER_RDG_BUFFER_SRV(StructuredBuffer<FRGB12BitBuffer>, InputRGB12Buffer)
-			SHADER_PARAMETER(uint32, HorizontalElementCount)
-			SHADER_PARAMETER(uint32, VerticalElementCount)
+			SHADER_PARAMETER(uint32, TotalElementCount)
+			SHADER_PARAMETER(uint32, OutputTextureSizeX)
 			SHADER_PARAMETER_RDG_TEXTURE_UAV(RWTexture2D, OutTexture)
 		END_SHADER_PARAMETER_STRUCT()
 
@@ -431,7 +426,7 @@ namespace UE::RivermaxShaders
 		}
 
 		/** Allocates and setup shader parameter in the incoming graph builder */
-		RIVERMAXRENDERING_API FRGB12BitToRGBA12CS::FParameters* AllocateAndSetParameters(FRDGBuilder& GraphBuilder, FRDGBufferRef RGBBuffer, FRDGTextureRef OutputTexture, int32 BufferElementsPerRow, int32 BufferLineCount);
+		RIVERMAXRENDERING_API FRGB12BitToRGBA12CS::FParameters* AllocateAndSetParameters(FRDGBuilder& GraphBuilder, FRDGBufferRef RGBBuffer, FRDGTextureRef OutputTexture, int32 TotalElementCount, int32 BufferLineCount);
 	};
 
 	/**
@@ -455,13 +450,12 @@ namespace UE::RivermaxShaders
 		BEGIN_SHADER_PARAMETER_STRUCT(FParameters, )
 			SHADER_PARAMETER_RDG_TEXTURE(Texture2D, InputTexture)
 			SHADER_PARAMETER_SAMPLER(SamplerState, InputSampler)
+			SHADER_PARAMETER(uint32, CapturedSizeX)
 			SHADER_PARAMETER(float, OnePixelDeltaX)
 			SHADER_PARAMETER(float, OnePixelDeltaY)
 			SHADER_PARAMETER(float, InputPixelOffsetX)
 			SHADER_PARAMETER(float, InputPixelOffsetY)
-			SHADER_PARAMETER(uint32, HorizontalElementCount)
-			SHADER_PARAMETER(uint32, VerticalElementCount)
-			SHADER_PARAMETER(uint32, InputTexturePixelsPerThread)
+			SHADER_PARAMETER(uint32, TotalElementCount)
 			SHADER_PARAMETER_RDG_BUFFER_UAV(RWStructuredBuffer<FRGB16fBuffer>, OutRGB16fBuffer)
 		END_SHADER_PARAMETER_STRUCT()
 
@@ -474,7 +468,7 @@ namespace UE::RivermaxShaders
 		}
 
 		/** Allocates and setup shader parameter in the incoming graph builder */
-		RIVERMAXRENDERING_API FRGBToRGB16fCS::FParameters* AllocateAndSetParameters(FRDGBuilder& GraphBuilder, FRDGTextureRef RGBATexture, const FIntPoint& SourceSize, const FIntRect& SourceViewRect, const FIntPoint& OutputSize, FRDGBufferRef OutputBuffer);
+		RIVERMAXRENDERING_API FRGBToRGB16fCS::FParameters* AllocateAndSetParameters(FRDGBuilder& GraphBuilder, FRDGTextureRef SourceTexture, const FIntPoint& CaptureSize, const FIntRect& SourceViewRect, const FIntPoint& OutputSize, FRDGBufferRef OutputBuffer);
 	};
 
 	/**
@@ -499,8 +493,8 @@ namespace UE::RivermaxShaders
 
 		BEGIN_SHADER_PARAMETER_STRUCT(FParameters, )
 			SHADER_PARAMETER_RDG_BUFFER_SRV(StructuredBuffer<FRGB16fBuffer>, InputRGB16fBuffer)
-			SHADER_PARAMETER(uint32, HorizontalElementCount)
-			SHADER_PARAMETER(uint32, VerticalElementCount)
+			SHADER_PARAMETER(uint32, TotalElementCount)
+			SHADER_PARAMETER(uint32, OutputTextureSizeX)
 			SHADER_PARAMETER_RDG_TEXTURE_UAV(RWTexture2D, OutTexture)
 		END_SHADER_PARAMETER_STRUCT()
 
@@ -513,7 +507,7 @@ namespace UE::RivermaxShaders
 		}
 
 		/** Allocates and setup shader parameter in the incoming graph builder */
-		RIVERMAXRENDERING_API FRGB16fBitToRGBA16fCS::FParameters* AllocateAndSetParameters(FRDGBuilder& GraphBuilder, FRDGBufferRef RGBBuffer, FRDGTextureRef OutputTexture, int32 BufferElementsPerRow, int32 BufferLineCount);
+		RIVERMAXRENDERING_API FRGB16fBitToRGBA16fCS::FParameters* AllocateAndSetParameters(FRDGBuilder& GraphBuilder, FRDGBufferRef RGBBuffer, FRDGTextureRef OutputTexture, int32 TotalElementCount, int32 BufferLineCount);
 	};
 }
 
diff --git a/Engine/Plugins/VirtualProduction/Rivermax/RivermaxMedia/Source/RivermaxMedia/Private/RivermaxMediaCapture.cpp b/Engine/Plugins/VirtualProduction/Rivermax/RivermaxMedia/Source/RivermaxMedia/Private/RivermaxMediaCapture.cpp
--- a/Engine/Plugins/VirtualProduction/Rivermax/RivermaxMedia/Source/RivermaxMedia/Private/RivermaxMediaCapture.cpp
+++ b/Engine/Plugins/VirtualProduction/Rivermax/RivermaxMedia/Source/RivermaxMedia/Private/RivermaxMediaCapture.cpp
@@ -96,7 +96,7 @@ namespace UE::RivermaxMediaCaptureUtil
 		return FIntPoint(AlignedHorizontalResolution, InSize.Y);
 	}
 
-	void GetOutputEncodingInfo(ERivermaxMediaOutputPixelFormat InPixelFormat, const FIntPoint& InSize, uint32& OutBytesPerElement, uint32& ElementsPerRow)
+	void GetOutputEncodingInfo(ERivermaxMediaOutputPixelFormat InPixelFormat, const FIntPoint& InSize, uint32& OutBytesPerElement, uint32& OutElementsPerFrame)
 	{
 		using namespace UE::RivermaxCore;
 
@@ -106,8 +106,9 @@ namespace UE::RivermaxMediaCaptureUtil
 		// Compute horizontal byte count (stride) of aligned resolution
 		const uint32 PixelAlignment = Info.PixelGroupCoverage;
 		const FIntPoint AlignedResolution = GetAlignedResolution(Info, InSize);
-		const uint32 PixelGroupCount = AlignedResolution.X / PixelAlignment;
-		const uint32 HorizontalByteCount = PixelGroupCount * Info.PixelGroupSize;
+		const uint32 PixelCount = AlignedResolution.X * AlignedResolution.Y;
+		const uint32 PixelGroupCount = PixelCount / PixelAlignment;
+		const uint32 FrameByteCount = PixelGroupCount * Info.PixelGroupSize;
 
 		switch (InPixelFormat)
 		{
@@ -144,7 +145,7 @@ namespace UE::RivermaxMediaCaptureUtil
 		}
 
 		// Shader encoding might not align with pixel group size so we need to have enough elements to represent the last pixel group
-		ElementsPerRow = FMath::CeilToInt32((float)HorizontalByteCount / OutBytesPerElement);
+		OutElementsPerFrame = FMath::CeilToInt32((float)FrameByteCount / OutBytesPerElement);
 	}
 }
 
@@ -305,7 +306,7 @@ bool URivermaxMediaCapture::ConfigureStream(URivermaxMediaOutput* InMediaOutput,
 
 	OutOptions.StreamAddress = InMediaOutput->StreamAddress;
 	OutOptions.Port = InMediaOutput->Port;
-	OutOptions.Resolution = InMediaOutput->Resolution;
+	OutOptions.Resolution = GetDesiredSize();
 	OutOptions.FrameRate = InMediaOutput->FrameRate;
 	OutOptions.NumberOfBuffers = InMediaOutput->NumberOfTextureBuffers;
 	OutOptions.bUseGPUDirect = InMediaOutput->bUseGPUDirect;
@@ -317,7 +318,7 @@ bool URivermaxMediaCapture::ConfigureStream(URivermaxMediaOutput* InMediaOutput,
 
 	OutOptions.PixelFormat = UE::RivermaxMediaUtils::Private::MediaOutputPixelFormatToRivermaxSamplingType(InMediaOutput->PixelFormat);
 	const FVideoFormatInfo Info = FStandardVideoFormat::GetVideoFormatInfo(OutOptions.PixelFormat);
-	OutOptions.AlignedResolution = UE::RivermaxMediaCaptureUtil::GetAlignedResolution(Info, InMediaOutput->Resolution);
+	OutOptions.AlignedResolution = UE::RivermaxMediaCaptureUtil::GetAlignedResolution(Info, OutOptions.Resolution);
 
 	return true;
 }
@@ -436,9 +437,9 @@ FIntPoint URivermaxMediaCapture::GetCustomOutputSize(const FIntPoint& InSize) co
 	URivermaxMediaOutput* RivermaxOutput = CastChecked<URivermaxMediaOutput>(MediaOutput);
 
 	uint32 BytesPerElement = 0;
-	uint32 ElementsPerRow = 0;
-	UE::RivermaxMediaCaptureUtil::GetOutputEncodingInfo(RivermaxOutput->PixelFormat, InSize, BytesPerElement, ElementsPerRow);
-	return FIntPoint(ElementsPerRow, InSize.Y);
+	uint32 ElementsPerFrame = 0;
+	UE::RivermaxMediaCaptureUtil::GetOutputEncodingInfo(RivermaxOutput->PixelFormat, InSize, BytesPerElement, ElementsPerFrame);
+	return FIntPoint(ElementsPerFrame, 1);
 }
 
 EMediaCaptureResourceType URivermaxMediaCapture::GetCustomOutputResourceType() const
@@ -466,9 +467,9 @@ FRDGBufferDesc URivermaxMediaCapture::GetCustomBufferDescription(const FIntPoint
 	URivermaxMediaOutput* RivermaxOutput = CastChecked<URivermaxMediaOutput>(MediaOutput);
 
 	uint32 BytesPerElement = 0;
-	uint32 ElementsPerRow = 0;
-	UE::RivermaxMediaCaptureUtil::GetOutputEncodingInfo(RivermaxOutput->PixelFormat, InDesiredSize, BytesPerElement, ElementsPerRow);
-	FRDGBufferDesc Desc = FRDGBufferDesc::CreateStructuredDesc(BytesPerElement, ElementsPerRow * InDesiredSize.Y);
+	uint32 ElementsPerFrame = 0;
+	UE::RivermaxMediaCaptureUtil::GetOutputEncodingInfo(RivermaxOutput->PixelFormat, InDesiredSize, BytesPerElement, ElementsPerFrame);
+	FRDGBufferDesc Desc = FRDGBufferDesc::CreateStructuredDesc(BytesPerElement, ElementsPerFrame);
 	
 	// Required when GPUDirect using CUDA will be involved
 	Desc.Usage |= EBufferUsageFlags::Shared;
@@ -488,15 +489,15 @@ void URivermaxMediaCapture::OnCustomCapture_RenderingThread(FRDGBuilder& GraphBu
 	// Rectangle area to use from source. This is used when source render target is bigger than output resolution
 	const FIntRect ViewRect(CopyInfo.GetSourceRect());
 	constexpr bool bDoLinearToSRGB = false;
-	const FIntPoint SourceSize = { InSourceTexture->Desc.Extent.X, InSourceTexture->Desc.Extent.Y };
-	const FIntVector GroupCount = FComputeShaderUtils::GetGroupCount(DesiredOutputSize, FComputeShaderUtils::kGolden2DGroupSize);
+	const FIntPoint CaptureSize = Options.AlignedResolution;
+	const FIntVector GroupCount = FComputeShaderUtils::GetGroupCount(DesiredOutputSize.X, 64);
 	
 	switch (RivermaxOutput->PixelFormat)
 	{
 	case ERivermaxMediaOutputPixelFormat::PF_8BIT_YUV422:
 	{
 		TShaderMapRef<FRGBToYUV8Bit422CS> ComputeShader(GlobalShaderMap);
-		FRGBToYUV8Bit422CS::FParameters* Parameters = ComputeShader->AllocateAndSetParameters(GraphBuilder, InSourceTexture, SourceSize, ViewRect, DesiredOutputSize, MediaShaders::RgbToYuvRec709Scaled, MediaShaders::YUVOffset8bits, bDoLinearToSRGB, OutputBuffer);
+		FRGBToYUV8Bit422CS::FParameters* Parameters = ComputeShader->AllocateAndSetParameters(GraphBuilder, InSourceTexture, CaptureSize, ViewRect, DesiredOutputSize, MediaShaders::RgbToYuvRec709Scaled, MediaShaders::YUVOffset8bits, bDoLinearToSRGB, OutputBuffer);
 
 		FComputeShaderUtils::AddPass(
 			GraphBuilder
@@ -510,7 +511,7 @@ void URivermaxMediaCapture::OnCustomCapture_RenderingThread(FRDGBuilder& GraphBu
 	case ERivermaxMediaOutputPixelFormat::PF_10BIT_YUV422:
 	{
 		TShaderMapRef<FRGBToYUV10Bit422LittleEndianCS> ComputeShader(GlobalShaderMap);
-		FRGBToYUV10Bit422LittleEndianCS::FParameters* Parameters = ComputeShader->AllocateAndSetParameters(GraphBuilder, InSourceTexture, SourceSize, ViewRect, DesiredOutputSize, MediaShaders::RgbToYuvRec709Scaled, MediaShaders::YUVOffset10bits, bDoLinearToSRGB, OutputBuffer);
+		FRGBToYUV10Bit422LittleEndianCS::FParameters* Parameters = ComputeShader->AllocateAndSetParameters(GraphBuilder, InSourceTexture, CaptureSize, ViewRect, DesiredOutputSize, MediaShaders::RgbToYuvRec709Scaled, MediaShaders::YUVOffset10bits, bDoLinearToSRGB, OutputBuffer);
 
 		FComputeShaderUtils::AddPass(
 			GraphBuilder
@@ -524,7 +525,7 @@ void URivermaxMediaCapture::OnCustomCapture_RenderingThread(FRDGBuilder& GraphBu
 	case ERivermaxMediaOutputPixelFormat::PF_8BIT_RGB:
 	{
 		TShaderMapRef<FRGBToRGB8BitCS> ComputeShader(GlobalShaderMap);
-		FRGBToRGB8BitCS::FParameters* Parameters = ComputeShader->AllocateAndSetParameters(GraphBuilder, InSourceTexture, SourceSize, ViewRect, DesiredOutputSize, OutputBuffer);
+		FRGBToRGB8BitCS::FParameters* Parameters = ComputeShader->AllocateAndSetParameters(GraphBuilder, InSourceTexture, CaptureSize, ViewRect, DesiredOutputSize, OutputBuffer);
 
 		FComputeShaderUtils::AddPass(
 			GraphBuilder
@@ -538,7 +539,7 @@ void URivermaxMediaCapture::OnCustomCapture_RenderingThread(FRDGBuilder& GraphBu
 	case ERivermaxMediaOutputPixelFormat::PF_10BIT_RGB:
 	{
 		TShaderMapRef<FRGBToRGB10BitCS> ComputeShader(GlobalShaderMap);
-		FRGBToRGB10BitCS::FParameters* Parameters = ComputeShader->AllocateAndSetParameters(GraphBuilder, InSourceTexture, SourceSize, ViewRect, DesiredOutputSize, OutputBuffer);
+		FRGBToRGB10BitCS::FParameters* Parameters = ComputeShader->AllocateAndSetParameters(GraphBuilder, InSourceTexture, CaptureSize, ViewRect, DesiredOutputSize, OutputBuffer);
 
 		FComputeShaderUtils::AddPass(
 			GraphBuilder
@@ -552,7 +553,7 @@ void URivermaxMediaCapture::OnCustomCapture_RenderingThread(FRDGBuilder& GraphBu
 	case ERivermaxMediaOutputPixelFormat::PF_12BIT_RGB:
 	{
 		TShaderMapRef<FRGBToRGB12BitCS> ComputeShader(GlobalShaderMap);
-		FRGBToRGB12BitCS::FParameters* Parameters = ComputeShader->AllocateAndSetParameters(GraphBuilder, InSourceTexture, SourceSize, ViewRect, DesiredOutputSize, OutputBuffer);
+		FRGBToRGB12BitCS::FParameters* Parameters = ComputeShader->AllocateAndSetParameters(GraphBuilder, InSourceTexture, CaptureSize, ViewRect, DesiredOutputSize, OutputBuffer);
 
 		FComputeShaderUtils::AddPass(
 			GraphBuilder
@@ -566,7 +567,7 @@ void URivermaxMediaCapture::OnCustomCapture_RenderingThread(FRDGBuilder& GraphBu
 	case ERivermaxMediaOutputPixelFormat::PF_FLOAT16_RGB:
 	{
 		TShaderMapRef<FRGBToRGB16fCS> ComputeShader(GlobalShaderMap);
-		FRGBToRGB16fCS::FParameters* Parameters = ComputeShader->AllocateAndSetParameters(GraphBuilder, InSourceTexture, SourceSize, ViewRect, DesiredOutputSize, OutputBuffer);
+		FRGBToRGB16fCS::FParameters* Parameters = ComputeShader->AllocateAndSetParameters(GraphBuilder, InSourceTexture, CaptureSize, ViewRect, DesiredOutputSize, OutputBuffer);
 
 		FComputeShaderUtils::AddPass(
 			GraphBuilder
diff --git a/Engine/Plugins/VirtualProduction/Rivermax/RivermaxMedia/Source/RivermaxMedia/Private/RivermaxMediaOutput.cpp b/Engine/Plugins/VirtualProduction/Rivermax/RivermaxMedia/Source/RivermaxMedia/Private/RivermaxMediaOutput.cpp
--- a/Engine/Plugins/VirtualProduction/Rivermax/RivermaxMedia/Source/RivermaxMedia/Private/RivermaxMediaOutput.cpp
+++ b/Engine/Plugins/VirtualProduction/Rivermax/RivermaxMedia/Source/RivermaxMedia/Private/RivermaxMediaOutput.cpp
@@ -21,7 +21,12 @@ bool URivermaxMediaOutput::Validate(FString& OutFailureReason) const
 
 FIntPoint URivermaxMediaOutput::GetRequestedSize() const
 {
-	return Resolution;
+	if (bOverrideResolution)
+	{
+		return Resolution;
+	}
+
+	return UMediaOutput::RequestCaptureSourceSize;
 }
 
 EPixelFormat URivermaxMediaOutput::GetRequestedPixelFormat() const
diff --git a/Engine/Plugins/VirtualProduction/Rivermax/RivermaxMedia/Source/RivermaxMedia/Private/RivermaxMediaPlayer.cpp b/Engine/Plugins/VirtualProduction/Rivermax/RivermaxMedia/Source/RivermaxMedia/Private/RivermaxMediaPlayer.cpp
--- a/Engine/Plugins/VirtualProduction/Rivermax/RivermaxMedia/Source/RivermaxMedia/Private/RivermaxMediaPlayer.cpp
+++ b/Engine/Plugins/VirtualProduction/Rivermax/RivermaxMedia/Source/RivermaxMedia/Private/RivermaxMediaPlayer.cpp
@@ -87,6 +87,8 @@ namespace UE::RivermaxMedia
 			DesiredPixelFormat = (ERivermaxMediaSourcePixelFormat)Options->GetMediaOption(RivermaxMediaOption::PixelFormat, (int64)ERivermaxMediaSourcePixelFormat::RGB_8bit);
 			const bool bUseZeroLatency = Options->GetMediaOption(RivermaxMediaOption::ZeroLatency, true);
 			FrameLatency = bUseZeroLatency ? 0 : 1;
+			const bool bOverrideResolution = Options->GetMediaOption(RivermaxMediaOption::OverrideResolution, false);
+			bFollowsStreamResolution = !bOverrideResolution;
 		}
 
 		bUseVideo = true;
@@ -97,6 +99,15 @@ namespace UE::RivermaxMedia
 			InputStream = Module->CreateInputStream();
 		}
 
+		// If we are not following the stream resolution, make it the video track format and then reset to go through a format change once
+		if (!bFollowsStreamResolution)
+		{
+			StreamResolution = StreamOptions.EnforcedResolution;
+		}
+		VideoTrackFormat.Dim = FIntPoint::ZeroValue;
+
+		FrameTracking.bWasFrameRequested = false;
+		FrameTracking.LastFrameRendered.Reset();
 		CurrentState = EMediaState::Preparing;
 		RivermaxThreadNewState = EMediaState::Preparing;
 		
@@ -112,9 +123,6 @@ namespace UE::RivermaxMedia
 		// Setup our different supported channels based on source settings
 		SetupSampleChannels();
 
-		// Allocate different buffers, samples we will be using
-		AllocateBuffers();
-
 		EventSink.ReceiveMediaEvent(EMediaEvent::MediaConnecting);
 
 #if WITH_EDITOR
@@ -137,8 +145,7 @@ namespace UE::RivermaxMedia
 	{
 		RivermaxThreadNewState = EMediaState::Closed;
 
-		// Flush any rendering activity to be sure we can move on with clearing resources. 
-		FlushRenderingCommands();
+		WaitForPendingTasks();
 
 		if (InputStream)
 		{
@@ -146,19 +153,6 @@ namespace UE::RivermaxMedia
 			InputStream.Reset();
 		}
 
-		// Wait for all pending tasks to complete. They should all complete at some point but add a timeout as a last resort. 
-		constexpr double TimeoutSeconds = 2.0;
-		const double StartTimeSeconds = FPlatformTime::Seconds();
-		while (TasksInFlight > 0)
-		{
-			FPlatformProcess::SleepNoStats(SleepTimeSeconds);
-
-			if ((FPlatformTime::Seconds() - StartTimeSeconds) > TimeoutSeconds)
-			{
-				break;
-			}
-		}
-
 		for(const TSharedPtr<FRivermaxSampleWrapper>& Sample : SamplePool)
 		{
 			Sample->Sample.Reset();
@@ -201,7 +195,7 @@ namespace UE::RivermaxMedia
 		TRACE_CPUPROFILER_EVENT_SCOPE(FRivermaxMediaPlayer::OnVideoFrameRequested)
 
 		// If video is not playing, no need to provide samples when requested
-		if (RivermaxThreadNewState != EMediaState::Playing)
+		if (!IsReadyToPlay())
 		{
 			return false;
 		}
@@ -279,7 +273,8 @@ namespace UE::RivermaxMedia
 	{
 		TRACE_CPUPROFILER_EVENT_SCOPE(FRivermaxMediaPlayer::OnVideoFrameReceived)
 
-		if (RivermaxThreadNewState != EMediaState::Playing)
+
+		if (!IsReadyToPlay())
 		{
 			return;
 		}
@@ -336,6 +331,42 @@ namespace UE::RivermaxMedia
 		{
 			return;
 		}
+		
+		// Cache current stream detection, it could change while we are applying it
+		FIntPoint CachedStreamResolution;
+		{
+			FScopeLock Lock(&StreamResolutionCriticalSection);
+			CachedStreamResolution = StreamResolution;
+		}
+
+		if (VideoTrackFormat.Dim != CachedStreamResolution)
+		{
+			UE_LOG(LogRivermaxMedia, Log, TEXT("Player needs to apply newly detected stream resolution : %dx%d"), CachedStreamResolution.X, CachedStreamResolution.Y);
+			
+			// Reset some frame tracking info while changing resolution
+			FrameTracking.bWasFrameRequested = false;
+			FrameTracking.LastFrameRendered.Reset();
+
+			{
+				WaitForPendingTasks();
+
+				// Cleanup allocated ressources for the current resolution
+				for (const TSharedPtr<FRivermaxSampleWrapper>& Sample : SamplePool)
+				{
+					Sample->Sample.Reset();
+					Sample->SampleConversionFence.SafeRelease();
+				}
+				SamplePool.Empty();
+
+				RivermaxThreadCurrentTextureSample.Reset();
+				SkippedFrames.Empty();
+
+				AllocateBuffers(CachedStreamResolution);
+				
+				VideoTrackFormat.Dim = CachedStreamResolution;
+			}
+
+		}
 
 		TickTimeManagement();
 	}
@@ -345,6 +376,8 @@ namespace UE::RivermaxMedia
 	 *****************************************************************************/
 	void FRivermaxMediaPlayer::ProcessFrame()
 	{
+		TRACE_CPUPROFILER_EVENT_SCOPE(RivermaxPlayerProcessFrame);
+
 		// Don't start making frame available until one is being received
 		if (FrameTracking.bWasFrameRequested == false)
 		{
@@ -355,8 +388,8 @@ namespace UE::RivermaxMedia
 		FSampleConfigurationArgs Args;
 		Args.bInIsSRGBInput = bIsSRGBInput;
 		Args.FrameRate = VideoFrameRate;
-		Args.Width = StreamOptions.Resolution.X;
-		Args.Height = StreamOptions.Resolution.Y;
+		Args.Width = VideoTrackFormat.Dim.X;
+		Args.Height = VideoTrackFormat.Dim.Y;
 		Args.Player = StaticCastSharedPtr<FRivermaxMediaPlayer>(AsShared().ToSharedPtr());
 		Args.SampleFormat = DesiredPixelFormat;
 		Args.Time = FTimespan(GFrameCounter);
@@ -426,21 +459,18 @@ namespace UE::RivermaxMedia
 		StreamOptions.StreamAddress = Options->GetMediaOption(RivermaxMediaOption::StreamAddress, FString());
 		StreamOptions.Port = Options->GetMediaOption(RivermaxMediaOption::Port, (int64)0);
 		StreamOptions.bUseGPUDirect = Options->GetMediaOption(RivermaxMediaOption::UseGPUDirect, false);
-		StreamOptions.Resolution = VideoTrackFormat.Dim;
 		StreamOptions.FrameRate = VideoFrameRate;
-
-		// We need to adjust horizontal resolution if it's not aligned to pgroup pixel coverage
-		// RTP header has to describe a full pgroup even if it's outside of effective resolution.
 		StreamOptions.PixelFormat = UE::RivermaxMediaUtils::Private::MediaSourcePixelFormatToRivermaxSamplingType(DesiredPixelFormat);
 		const FVideoFormatInfo FormatInfo = FStandardVideoFormat::GetVideoFormatInfo(StreamOptions.PixelFormat);
 		const uint32 PixelAlignment = FormatInfo.PixelGroupCoverage;
-		const uint32 AlignedHorizontalResolution = (StreamOptions.Resolution.X % PixelAlignment) ? StreamOptions.Resolution.X + (PixelAlignment - (StreamOptions.Resolution.X % PixelAlignment)) : StreamOptions.Resolution.X;
-		StreamOptions.AlignedResolution = FIntPoint(AlignedHorizontalResolution, StreamOptions.Resolution.Y);
+		const uint32 AlignedHorizontalResolution = (VideoTrackFormat.Dim.X % PixelAlignment) ? VideoTrackFormat.Dim.X + (PixelAlignment - (VideoTrackFormat.Dim.X % PixelAlignment)) : VideoTrackFormat.Dim.X;
+		StreamOptions.EnforcedResolution = FIntPoint(AlignedHorizontalResolution, VideoTrackFormat.Dim.Y);
+		StreamOptions.bEnforceVideoFormat = !bFollowsStreamResolution;
 
 		return true;
 	}
 
-	void FRivermaxMediaPlayer::AllocateBuffers()
+	void FRivermaxMediaPlayer::AllocateBuffers(const FIntPoint& InResolution)
 	{
 		using namespace UE::RivermaxCore;
 		using namespace UE::RivermaxMediaUtils::Private;
@@ -448,7 +478,7 @@ namespace UE::RivermaxMedia
 		// Take care of the common buffers first
 		{
 			// Create the common texture we are going to use
-			const FSourceBufferDesc BufferDescription = GetBufferDescription(StreamOptions.AlignedResolution, DesiredPixelFormat);
+			const FSourceBufferDesc BufferDescription = GetBufferDescription(InResolution, DesiredPixelFormat);
 			FRDGBufferDesc RDGBufferDesc = FRDGBufferDesc::CreateStructuredDesc(BufferDescription.BytesPerElement, BufferDescription.NumberOfElements);
 
 			// Required to share resource across different graphics API (DX, Cuda)
@@ -465,12 +495,14 @@ namespace UE::RivermaxMedia
 				});
 		}
 
+		SamplePool.Empty();
+
 		// Allocate our pool of samples where incoming ones will be written and chosen from
 		for (int32 Index = 0; Index < MaxNumVideoFrameBuffer; Index++)
 		{
 			TSharedPtr<FRivermaxSampleWrapper> NewWrapper = MakeShared<FRivermaxSampleWrapper>();
 			NewWrapper->Sample = MakeShared<FRivermaxMediaTextureSample>();
-			NewWrapper->Sample->InitializeGPUBuffer(StreamOptions.AlignedResolution, DesiredPixelFormat);
+			NewWrapper->Sample->InitializeGPUBuffer(InResolution, DesiredPixelFormat);
 			NewWrapper->SampleConversionFence = RHICreateGPUFence(*FString::Printf(TEXT("RmaxConversionDoneFence_%02d"), Index));
 			NewWrapper->ReceptionState = ESampleReceptionState::Available;
 			SamplePool.Add(MoveTemp(NewWrapper));
@@ -517,7 +549,7 @@ namespace UE::RivermaxMedia
 		}
 
 		// If it's the first time we are rendering, make previous frame number as available since we will never get to them
-		if (FrameTracking.bHasRendered == false)
+		if (FrameTracking.LastFrameRendered.IsSet() && NextFrameExpectations.FrameNumber > (FrameTracking.LastFrameRendered.GetValue() + 1))
 		{
 			for (const TSharedPtr<FRivermaxSampleWrapper>& Frame : SamplePool)
 			{
@@ -526,13 +558,13 @@ namespace UE::RivermaxMedia
 					if (Frame->FrameNumber < NextFrameExpectations.FrameNumber)
 					{
 						Frame->ReceptionState = ESampleReceptionState::Available;
-						UE_LOG(LogRivermaxMedia, Verbose, TEXT("Making frame %u as available since it will never be processed."), GFrameCounterRenderThread, Frame->FrameNumber);
+						UE_LOG(LogRivermaxMedia, Verbose, TEXT("Making frame %u as available since it will never be processed. Last render = %u, next render = %u"), Frame->FrameNumber, FrameTracking.LastFrameRendered.GetValue(), NextFrameExpectations.FrameNumber);
 					}
 				}
 			}
-
-			FrameTracking.bHasRendered = true;
 		}
+		
+		FrameTracking.LastFrameRendered = NextFrameExpectations.FrameNumber;
 
 		// Verify if the frame we will use for rendering is still being rendered for the previous one.
 		if (SamplePool[NextFrameExpectations.FrameIndex]->bIsPendingRendering)
@@ -753,7 +785,7 @@ namespace UE::RivermaxMedia
 					FPlatformProcess::SleepNoStats(SleepTimeSeconds);
 					if ((FPlatformTime::Seconds() - StartTimeSeconds) > TimeoutSeconds)
 					{
-						UE_LOG(LogRivermaxMedia, Error, TEXT("Timed out waiting for frame %u to be available to receive into."), ExpectedFrame->FrameNumber);
+						UE_LOG(LogRivermaxMedia, Error, TEXT("Timed out waiting for frame %u to be rendered to receive frame %u into."), ExpectedFrame->FrameNumber, FrameInfo.FrameNumber);
 						break;
 					}
 				}
@@ -822,7 +854,7 @@ namespace UE::RivermaxMedia
 				while (true)
 				{
 					{
-						if (RivermaxThreadNewState != EMediaState::Playing)
+						if (!IsReadyToPlay())
 						{
 							break;
 						}
@@ -966,6 +998,50 @@ namespace UE::RivermaxMedia
 		}
 	}
 
+	void FRivermaxMediaPlayer::OnVideoFormatChanged(const FRivermaxInputVideoFormatChangedInfo& NewFormatInfo)
+	{
+		const ERivermaxMediaSourcePixelFormat NewFormat = UE::RivermaxMediaUtils::Private::RivermaxPixelFormatToMediaSourcePixelFormat(NewFormatInfo.PixelFormat);
+		const FIntPoint NewResolution = { (int32)NewFormatInfo.Width, (int32)NewFormatInfo.Height };
+		bool bNeedReinitializing = (NewFormatInfo.PixelFormat != StreamOptions.PixelFormat);
+		bNeedReinitializing |= (NewFormatInfo.Width != VideoTrackFormat.Dim.X || NewFormatInfo.Height != VideoTrackFormat.Dim.Y);
+		
+		UE_LOG(LogRivermaxMedia, Log, TEXT("New video format detected: %dx%d with pixel format '%s'"), NewResolution.X, NewResolution.Y, *UEnum::GetValueAsString(NewFormat));
+		
+		if (bNeedReinitializing && bFollowsStreamResolution)
+		{
+			FScopeLock Lock(&StreamResolutionCriticalSection);
+			StreamResolution = NewResolution;
+		}
+	}
+
+	bool FRivermaxMediaPlayer::IsReadyToPlay() const
+	{
+		if (RivermaxThreadNewState == EMediaState::Playing)
+		{
+			FScopeLock Lock(&StreamResolutionCriticalSection);
+			return StreamResolution == VideoTrackFormat.Dim;
+		}
+
+		return false;
+	}
+
+	void FRivermaxMediaPlayer::WaitForPendingTasks()
+	{
+		// Flush any rendering activity to be sure we can move on with clearing resources. 
+		FlushRenderingCommands();
+
+		// Wait for all pending tasks to complete. They should all complete at some point but add a timeout as a last resort. 
+		constexpr double TimeoutSeconds = 2.0;
+		const double StartTimeSeconds = FPlatformTime::Seconds();
+		while (TasksInFlight > 0)
+		{
+			FPlatformProcess::SleepNoStats(SleepTimeSeconds);
+			if ((FPlatformTime::Seconds() - StartTimeSeconds) > TimeoutSeconds)
+			{
+				break;
+			}
+		}
+	}
 }
 
 #undef LOCTEXT_NAMESPACE
diff --git a/Engine/Plugins/VirtualProduction/Rivermax/RivermaxMedia/Source/RivermaxMedia/Private/RivermaxMediaPlayer.h b/Engine/Plugins/VirtualProduction/Rivermax/RivermaxMedia/Source/RivermaxMedia/Private/RivermaxMediaPlayer.h
--- a/Engine/Plugins/VirtualProduction/Rivermax/RivermaxMedia/Source/RivermaxMedia/Private/RivermaxMediaPlayer.h
+++ b/Engine/Plugins/VirtualProduction/Rivermax/RivermaxMedia/Source/RivermaxMedia/Private/RivermaxMediaPlayer.h
@@ -72,6 +72,7 @@ namespace  UE::RivermaxMedia
 		virtual void OnVideoFrameReceived(const FRivermaxInputVideoFrameDescriptor& FrameInfo, const FRivermaxInputVideoFrameReception& ReceivedVideoFrame) override;
 		virtual void OnVideoFrameReceptionError(const FRivermaxInputVideoFrameDescriptor& FrameInfo) override;
 		virtual void OnStreamError() override;
+		virtual void OnVideoFormatChanged(const FRivermaxInputVideoFormatChangedInfo& NewFormatInfo) override;
 		//~ End IRivermaxInputStreamListener interface
 
 	protected:
@@ -94,7 +95,7 @@ namespace  UE::RivermaxMedia
 		bool ConfigureStream(const IMediaOptions* Options);
 
 		/** Allocates the sample pool used to receive incoming data */
-		void AllocateBuffers();
+		void AllocateBuffers(const FIntPoint& InResolution);
 
 		/** Wrapper struct holding information about frame expected to be rendered */
 		struct FFrameExpectation
@@ -128,6 +129,13 @@ namespace  UE::RivermaxMedia
 		/** Looks at last frame rendered and tries to clean skipped frames container */
 		void TryClearSkippedInterval(uint32 LastFrameRendered);
 
+		/** Whether player is ready to play */
+		bool IsReadyToPlay() const;
+
+		/** Waits for tasks in flight and flushes render commands before cleaning our ressources */
+		void WaitForPendingTasks();
+
+
 	private:
 
 		enum class ESampleReceptionState : uint8
@@ -201,8 +209,8 @@ namespace  UE::RivermaxMedia
 			/** Whether a frame has been requested yet. Used to detect start of the stream and valid frame number expectations */
 			bool bWasFrameRequested = false;
 			
-			/** Whether a call to render was made. Used to detect start of rendering and discard incoming frames that will never be used */
-			bool bHasRendered = false;
+			/** Last frame number sent to render. Used to release received frames that will never be picked up early on */
+			TOptional<uint32> LastFrameRendered;
 
 			/** Last frame expectations used */
 			FFrameExpectation LastFrameExpectation;
@@ -238,6 +246,15 @@ namespace  UE::RivermaxMedia
 
 		/** Time to sleep when waiting for an operation to complete */
 		static constexpr double SleepTimeSeconds = 50.0 * 1E-6;
+
+		/** Critical section used when accessing stream resolution and detect pending changes */
+		mutable FCriticalSection StreamResolutionCriticalSection;
+
+		/** Resolution detected by our stream. */
+		FIntPoint StreamResolution = FIntPoint::ZeroValue;
+		
+		/** Whether the player follows resolution detected by our stream, adjusting texture size as required */
+		bool bFollowsStreamResolution = true;
 	};
 }
 
diff --git a/Engine/Plugins/VirtualProduction/Rivermax/RivermaxMedia/Source/RivermaxMedia/Private/RivermaxMediaSource.cpp b/Engine/Plugins/VirtualProduction/Rivermax/RivermaxMedia/Source/RivermaxMedia/Private/RivermaxMediaSource.cpp
--- a/Engine/Plugins/VirtualProduction/Rivermax/RivermaxMedia/Source/RivermaxMedia/Private/RivermaxMediaSource.cpp
+++ b/Engine/Plugins/VirtualProduction/Rivermax/RivermaxMedia/Source/RivermaxMedia/Private/RivermaxMediaSource.cpp
@@ -24,6 +24,11 @@ bool URivermaxMediaSource::GetMediaOption(const FName& Key, bool DefaultValue) c
 	{
 		return bUseZeroLatency;
 	}
+	if (Key == RivermaxMediaOption::OverrideResolution)
+	{
+		return bOverrideResolution;
+	}
+
 	return Super::GetMediaOption(Key, DefaultValue);
 }
 
diff --git a/Engine/Plugins/VirtualProduction/Rivermax/RivermaxMedia/Source/RivermaxMedia/Private/RivermaxMediaSourceOptions.h b/Engine/Plugins/VirtualProduction/Rivermax/RivermaxMedia/Source/RivermaxMedia/Private/RivermaxMediaSourceOptions.h
--- a/Engine/Plugins/VirtualProduction/Rivermax/RivermaxMedia/Source/RivermaxMedia/Private/RivermaxMediaSourceOptions.h
+++ b/Engine/Plugins/VirtualProduction/Rivermax/RivermaxMedia/Source/RivermaxMedia/Private/RivermaxMediaSourceOptions.h
@@ -15,4 +15,5 @@ namespace RivermaxMediaOption
 	static const FName UseGPUDirect("UseGPUDirect");
 	static const FName PlayerMode("PlayerMode");
 	static const FName ZeroLatency("ZeroLatency");
+	static const FName OverrideResolution("OverrideResolution");
 }
diff --git a/Engine/Plugins/VirtualProduction/Rivermax/RivermaxMedia/Source/RivermaxMedia/Private/RivermaxMediaTextureSampleConverter.cpp b/Engine/Plugins/VirtualProduction/Rivermax/RivermaxMedia/Source/RivermaxMedia/Private/RivermaxMediaTextureSampleConverter.cpp
--- a/Engine/Plugins/VirtualProduction/Rivermax/RivermaxMedia/Source/RivermaxMedia/Private/RivermaxMediaTextureSampleConverter.cpp
+++ b/Engine/Plugins/VirtualProduction/Rivermax/RivermaxMedia/Source/RivermaxMedia/Private/RivermaxMediaTextureSampleConverter.cpp
@@ -84,8 +84,8 @@ bool FRivermaxMediaTextureSampleConverter::Convert(FTexture2DRHIRef& InDestinati
 		}
 
 		
-		const FIntPoint ProcessedOutputDimension = { (int32)SourceBufferDesc.ElementsPerRow, InDestinationTexture->GetDesc().Extent.Y };
-		const FIntVector GroupCount = FComputeShaderUtils::GetGroupCount(ProcessedOutputDimension, FComputeShaderUtils::kGolden2DGroupSize);
+		const FIntPoint ProcessedOutputDimension = { (int32)SourceBufferDesc.NumberOfElements, 1 };
+		const FIntVector GroupCount = FComputeShaderUtils::GetGroupCount(SourceBufferDesc.NumberOfElements, 64);
 		FGlobalShaderMap* GlobalShaderMap = GetGlobalShaderMap(GMaxRHIFeatureLevel);
 
 		//Configure shader and add conversion pass based on desired pixel format
@@ -101,8 +101,7 @@ bool FRivermaxMediaTextureSampleConverter::Convert(FTexture2DRHIRef& InDestinati
 			const FMatrix YUVToRGBMatrix = SamplePtr->GetYUVToRGBMatrix();
 			const FVector YUVOffset(MediaShaders::YUVOffset8bits);
 			TShaderMapRef<FYUV8Bit422ToRGBACS> ComputeShader(GlobalShaderMap, PermutationVector);
-			FYUV8Bit422ToRGBACS::FParameters* Parameters = ComputeShader->AllocateAndSetParameters(GraphBuilder, InputBuffer, OutputResource, YUVToRGBMatrix, YUVOffset, SourceBufferDesc.ElementsPerRow, ProcessedOutputDimension.Y);
-
+			FYUV8Bit422ToRGBACS::FParameters* Parameters = ComputeShader->AllocateAndSetParameters(GraphBuilder, InputBuffer, OutputResource, YUVToRGBMatrix, YUVOffset, ProcessedOutputDimension.X, ProcessedOutputDimension.Y);
 
 			FComputeShaderUtils::AddPass(
 				GraphBuilder
@@ -122,7 +121,7 @@ bool FRivermaxMediaTextureSampleConverter::Convert(FTexture2DRHIRef& InDestinati
 			const FMatrix YUVToRGBMatrix = SamplePtr->GetYUVToRGBMatrix();
 			const FVector YUVOffset(MediaShaders::YUVOffset10bits);
 			TShaderMapRef<FYUV10Bit422ToRGBACS> ComputeShader(GlobalShaderMap, PermutationVector);
-			FYUV10Bit422ToRGBACS::FParameters* Parameters = ComputeShader->AllocateAndSetParameters(GraphBuilder, InputBuffer, OutputResource, YUVToRGBMatrix, YUVOffset, SourceBufferDesc.ElementsPerRow, ProcessedOutputDimension.Y);
+			FYUV10Bit422ToRGBACS::FParameters* Parameters = ComputeShader->AllocateAndSetParameters(GraphBuilder, InputBuffer, OutputResource, YUVToRGBMatrix, YUVOffset, ProcessedOutputDimension.X, ProcessedOutputDimension.Y);
 
 			FComputeShaderUtils::AddPass(
 				GraphBuilder
@@ -140,8 +139,7 @@ bool FRivermaxMediaTextureSampleConverter::Convert(FTexture2DRHIRef& InDestinati
 			PermutationVector.Set<FRGB8BitToRGBA8CS::FSRGBToLinear>(SamplePtr->NeedsSRGBToLinearConversion());
 
 			TShaderMapRef<FRGB8BitToRGBA8CS> ComputeShader(GlobalShaderMap, PermutationVector);
-			FRGB8BitToRGBA8CS::FParameters* Parameters = ComputeShader->AllocateAndSetParameters(GraphBuilder, InputBuffer, OutputResource, SourceBufferDesc.ElementsPerRow, ProcessedOutputDimension.Y);
-
+			FRGB8BitToRGBA8CS::FParameters* Parameters = ComputeShader->AllocateAndSetParameters(GraphBuilder, InputBuffer, OutputResource, ProcessedOutputDimension.X, ProcessedOutputDimension.Y);
 
 			FComputeShaderUtils::AddPass(
 				GraphBuilder
@@ -159,7 +157,7 @@ bool FRivermaxMediaTextureSampleConverter::Convert(FTexture2DRHIRef& InDestinati
 			PermutationVector.Set<FRGB10BitToRGBA10CS::FSRGBToLinear>(SamplePtr->NeedsSRGBToLinearConversion());
 
 			TShaderMapRef<FRGB10BitToRGBA10CS> ComputeShader(GlobalShaderMap, PermutationVector);
-			FRGB10BitToRGBA10CS::FParameters* Parameters = ComputeShader->AllocateAndSetParameters(GraphBuilder, InputBuffer, OutputResource, SourceBufferDesc.ElementsPerRow, ProcessedOutputDimension.Y);
+			FRGB10BitToRGBA10CS::FParameters* Parameters = ComputeShader->AllocateAndSetParameters(GraphBuilder, InputBuffer, OutputResource, ProcessedOutputDimension.X, ProcessedOutputDimension.Y);
 
 			FComputeShaderUtils::AddPass(
 				GraphBuilder
@@ -177,7 +175,7 @@ bool FRivermaxMediaTextureSampleConverter::Convert(FTexture2DRHIRef& InDestinati
 			PermutationVector.Set<FRGB12BitToRGBA12CS::FSRGBToLinear>(SamplePtr->NeedsSRGBToLinearConversion());
 
 			TShaderMapRef<FRGB12BitToRGBA12CS> ComputeShader(GlobalShaderMap, PermutationVector);
-			FRGB12BitToRGBA12CS::FParameters* Parameters = ComputeShader->AllocateAndSetParameters(GraphBuilder, InputBuffer, OutputResource, SourceBufferDesc.ElementsPerRow, ProcessedOutputDimension.Y);
+			FRGB12BitToRGBA12CS::FParameters* Parameters = ComputeShader->AllocateAndSetParameters(GraphBuilder, InputBuffer, OutputResource, ProcessedOutputDimension.X, ProcessedOutputDimension.Y);
 
 			FComputeShaderUtils::AddPass(
 				GraphBuilder
@@ -195,7 +193,7 @@ bool FRivermaxMediaTextureSampleConverter::Convert(FTexture2DRHIRef& InDestinati
 			PermutationVector.Set<FRGB16fBitToRGBA16fCS::FSRGBToLinear>(SamplePtr->NeedsSRGBToLinearConversion());
 
 			TShaderMapRef<FRGB16fBitToRGBA16fCS> ComputeShader(GlobalShaderMap, PermutationVector);
-			FRGB16fBitToRGBA16fCS::FParameters* Parameters = ComputeShader->AllocateAndSetParameters(GraphBuilder, InputBuffer, OutputResource, SourceBufferDesc.ElementsPerRow, ProcessedOutputDimension.Y);
+			FRGB16fBitToRGBA16fCS::FParameters* Parameters = ComputeShader->AllocateAndSetParameters(GraphBuilder, InputBuffer, OutputResource, ProcessedOutputDimension.X, ProcessedOutputDimension.Y);
 
 			FComputeShaderUtils::AddPass(
 				GraphBuilder
diff --git a/Engine/Plugins/VirtualProduction/Rivermax/RivermaxMedia/Source/RivermaxMedia/Private/RivermaxMediaUtils.cpp b/Engine/Plugins/VirtualProduction/Rivermax/RivermaxMedia/Source/RivermaxMedia/Private/RivermaxMediaUtils.cpp
--- a/Engine/Plugins/VirtualProduction/Rivermax/RivermaxMedia/Source/RivermaxMedia/Private/RivermaxMediaUtils.cpp
+++ b/Engine/Plugins/VirtualProduction/Rivermax/RivermaxMedia/Source/RivermaxMedia/Private/RivermaxMediaUtils.cpp
@@ -94,8 +94,10 @@ namespace UE::RivermaxMediaUtils::Private
 
 		// Compute horizontal byte count (stride) of resolution
 		const uint32 PixelAlignment = Info.PixelGroupCoverage;
-		Description.PixelGroupCount = PixelAlignment > 0 ? InResolution.X / PixelAlignment : 0;
-		Description.BytesPerRow = Description.PixelGroupCount * Info.PixelGroupSize;
+		const uint32 PixelCount = InResolution.X * InResolution.Y;
+		ensureMsgf(PixelCount % PixelAlignment == 0, TEXT("Unaligned resolution (%dx%d) provided for pixel format %s"), InResolution.X, InResolution.Y, *UEnum::GetValueAsString(InPixelFormat));
+		const uint32 PixelGroupCount = PixelCount / PixelAlignment;
+		const uint32 FrameByteCount = PixelGroupCount * Info.PixelGroupSize;
 
 		switch (SamplingType)
 		{
@@ -137,9 +139,8 @@ namespace UE::RivermaxMediaUtils::Private
 		}
 
 		// Shader encoding might not align with pixel group size so we need to have enough elements to represent the last pixel group
-		Description.ElementsPerRow = Description.BytesPerRow / Description.BytesPerElement;
-		Description.ElementsPerRow += Description.BytesPerRow % Description.BytesPerElement != 0 ? 1 : 0;
-		Description.NumberOfElements = Description.ElementsPerRow * InResolution.Y;
+		Description.NumberOfElements = FrameByteCount / Description.BytesPerElement;
+		Description.NumberOfElements += FrameByteCount % Description.BytesPerElement != 0 ? 1 : 0;
 		
 		return Description;
 	}
@@ -165,4 +166,81 @@ namespace UE::RivermaxMediaUtils::Private
 		}
 		}
 	}
+
+	ERivermaxMediaSourcePixelFormat RivermaxPixelFormatToMediaSourcePixelFormat(UE::RivermaxCore::ESamplingType InSamplingType)
+	{
+		using namespace UE::RivermaxCore;
+
+		switch (InSamplingType)
+		{
+		case ESamplingType::YUV422_8bit:
+		{
+			return ERivermaxMediaSourcePixelFormat::YUV422_8bit;
+		}
+		case ESamplingType::YUV422_10bit:
+		{
+			return ERivermaxMediaSourcePixelFormat::YUV422_10bit;
+		}
+		case ESamplingType::RGB_8bit:
+		{
+			return ERivermaxMediaSourcePixelFormat::RGB_8bit;
+		}
+		case ESamplingType::RGB_10bit:
+		{
+			return ERivermaxMediaSourcePixelFormat::RGB_10bit;
+		}
+		case ESamplingType::RGB_12bit:
+		{
+			return ERivermaxMediaSourcePixelFormat::RGB_12bit;
+		}
+		case ESamplingType::RGB_16bitFloat:
+		{
+			return ERivermaxMediaSourcePixelFormat::RGB_16bit_Float;
+		}
+		default:
+		{
+			checkNoEntry();
+			return ERivermaxMediaSourcePixelFormat::RGB_10bit;
+		}
+		}
+	}
+
+	ERivermaxMediaOutputPixelFormat RivermaxPixelFormatToMediaOutputPixelFormat(UE::RivermaxCore::ESamplingType InSamplingType)
+	{
+		using namespace UE::RivermaxCore;
+
+		switch (InSamplingType)
+		{
+		case ESamplingType::YUV422_8bit:
+		{
+			return ERivermaxMediaOutputPixelFormat::PF_8BIT_YUV422;
+		}
+		case ESamplingType::YUV422_10bit:
+		{
+			return ERivermaxMediaOutputPixelFormat::PF_10BIT_YUV422;
+		}
+		case ESamplingType::RGB_8bit:
+		{
+			return ERivermaxMediaOutputPixelFormat::PF_8BIT_RGB;
+		}
+		case ESamplingType::RGB_10bit:
+		{
+			return ERivermaxMediaOutputPixelFormat::PF_10BIT_RGB;
+		}
+		case ESamplingType::RGB_12bit:
+		{
+			return ERivermaxMediaOutputPixelFormat::PF_12BIT_RGB;
+		}
+		case ESamplingType::RGB_16bitFloat:
+		{
+			return ERivermaxMediaOutputPixelFormat::PF_FLOAT16_RGB;
+		}
+		default:
+		{
+			checkNoEntry();
+			return ERivermaxMediaOutputPixelFormat::PF_10BIT_RGB;
+		}
+		}
+	}
+
 }
diff --git a/Engine/Plugins/VirtualProduction/Rivermax/RivermaxMedia/Source/RivermaxMedia/Private/RivermaxMediaUtils.h b/Engine/Plugins/VirtualProduction/Rivermax/RivermaxMedia/Source/RivermaxMedia/Private/RivermaxMediaUtils.h
--- a/Engine/Plugins/VirtualProduction/Rivermax/RivermaxMedia/Source/RivermaxMedia/Private/RivermaxMediaUtils.h
+++ b/Engine/Plugins/VirtualProduction/Rivermax/RivermaxMedia/Source/RivermaxMedia/Private/RivermaxMediaUtils.h
@@ -13,15 +13,14 @@ namespace UE::RivermaxMediaUtils::Private
 {
 	struct FSourceBufferDesc
 	{
-		uint32 PixelGroupCount = 0;
 		uint32 BytesPerElement = 0;
-		uint32 ElementsPerRow = 0;
-		uint32 BytesPerRow = 0;
 		uint32 NumberOfElements = 0;
 	};
 
 	UE::RivermaxCore::ESamplingType MediaOutputPixelFormatToRivermaxSamplingType(ERivermaxMediaOutputPixelFormat InPixelFormat);
 	UE::RivermaxCore::ESamplingType MediaSourcePixelFormatToRivermaxSamplingType(ERivermaxMediaSourcePixelFormat InPixelFormat);
+	ERivermaxMediaSourcePixelFormat RivermaxPixelFormatToMediaSourcePixelFormat(UE::RivermaxCore::ESamplingType InSamplingType);
+	ERivermaxMediaOutputPixelFormat RivermaxPixelFormatToMediaOutputPixelFormat(UE::RivermaxCore::ESamplingType InSamplingType);
 	UE::RivermaxCore::ERivermaxAlignmentMode MediaOutputAlignmentToRivermaxAlignment(ERivermaxMediaAlignmentMode InAlignmentMode);
 	FSourceBufferDesc GetBufferDescription(const FIntPoint& Resolution, ERivermaxMediaSourcePixelFormat InPixelFormat);
 
diff --git a/Engine/Plugins/VirtualProduction/Rivermax/RivermaxMedia/Source/RivermaxMedia/Public/RivermaxMediaCapture.h b/Engine/Plugins/VirtualProduction/Rivermax/RivermaxMedia/Source/RivermaxMedia/Public/RivermaxMediaCapture.h
--- a/Engine/Plugins/VirtualProduction/Rivermax/RivermaxMedia/Source/RivermaxMedia/Public/RivermaxMediaCapture.h
+++ b/Engine/Plugins/VirtualProduction/Rivermax/RivermaxMedia/Source/RivermaxMedia/Public/RivermaxMediaCapture.h
@@ -44,6 +44,10 @@ protected:
 	{
 		return true;
 	}
+	virtual bool SupportsAutoRestart() const
+	{
+		return true;
+	}
 
 	/** For custom conversion, methods that need to be overriden */
 	virtual FIntPoint GetCustomOutputSize(const FIntPoint& InSize) const override;
diff --git a/Engine/Plugins/VirtualProduction/Rivermax/RivermaxMedia/Source/RivermaxMedia/Public/RivermaxMediaOutput.h b/Engine/Plugins/VirtualProduction/Rivermax/RivermaxMedia/Source/RivermaxMedia/Public/RivermaxMediaOutput.h
--- a/Engine/Plugins/VirtualProduction/Rivermax/RivermaxMedia/Source/RivermaxMedia/Public/RivermaxMediaOutput.h
+++ b/Engine/Plugins/VirtualProduction/Rivermax/RivermaxMedia/Source/RivermaxMedia/Public/RivermaxMediaOutput.h
@@ -84,8 +84,12 @@ public:
 	UPROPERTY(EditAnywhere, Category = "Output", meta = (EditCondition = "AlignmentMode == ERivermaxMediaAlignmentMode::FrameCreation"))
 	bool bDoFrameCounterTimestamping = true;
 
+	/** If false, use the default source buffer size. If true, a specific resolution will be used. */
+	UPROPERTY(BlueprintReadWrite, EditAnywhere, Category = "Media", meta = (InlineEditConditionToggle))
+	bool bOverrideResolution = false;
+
 	/** Resolution of this output stream */
-	UPROPERTY(EditAnywhere, BlueprintReadWrite, Category="Settings")
+	UPROPERTY(EditAnywhere, BlueprintReadWrite, Category="Settings", meta = (EditCondition = "bOverrideResolution"))
 	FIntPoint Resolution = {1920, 1080};
 	
 	/** Frame rate of this output stream */
diff --git a/Engine/Plugins/VirtualProduction/Rivermax/RivermaxMedia/Source/RivermaxMedia/Public/RivermaxMediaSource.h b/Engine/Plugins/VirtualProduction/Rivermax/RivermaxMedia/Source/RivermaxMedia/Public/RivermaxMediaSource.h
--- a/Engine/Plugins/VirtualProduction/Rivermax/RivermaxMedia/Source/RivermaxMedia/Public/RivermaxMediaSource.h
+++ b/Engine/Plugins/VirtualProduction/Rivermax/RivermaxMedia/Source/RivermaxMedia/Public/RivermaxMediaSource.h
@@ -59,8 +59,12 @@ public:
 	UPROPERTY(EditAnywhere, BlueprintReadWrite, Category = "Player", meta = (EditCondition = "PlayerMode == ERivermaxPlayerMode::Framelock"))
 	bool bUseZeroLatency  = true;
 
+	/** If false, use the default source buffer size. If true, a specific resolution will be used. */
+	UPROPERTY(BlueprintReadWrite, EditAnywhere, Category = "Media", meta = (InlineEditConditionToggle))
+	bool bOverrideResolution = false;
+
 	/** Incoming stream video resolution */
-	UPROPERTY(EditAnywhere, BlueprintReadWrite, Category="Options")
+	UPROPERTY(EditAnywhere, BlueprintReadWrite, Category="Options", meta = (EditCondition = "bOverrideResolution"))
 	FIntPoint Resolution = {1920, 1080};
 	
 	/** Incoming stream video frame rate */
