- Adding more logs in rmax out stream for missing chunks
- Changing scheduling when we detect stalls
- Adding TRO override cvar

diff --git a/Engine/Plugins/VirtualProduction/Rivermax/RivermaxCore/Source/RivermaxCore/Private/RivermaxOutputStream.cpp b/Engine/Plugins/VirtualProduction/Rivermax/RivermaxCore/Source/RivermaxCore/Private/RivermaxOutputStream.cpp
--- a/Engine/Plugins/VirtualProduction/Rivermax/RivermaxCore/Source/RivermaxCore/Private/RivermaxOutputStream.cpp
+++ b/Engine/Plugins/VirtualProduction/Rivermax/RivermaxCore/Source/RivermaxCore/Private/RivermaxOutputStream.cpp
@@ -51,6 +51,11 @@ namespace UE::RivermaxCore::Private
 		TEXT("Enables bigger packet sizes to maximize utilisation of potential UDP packet. If not enabled, packet size will be aligned with HD/4k sizes"),
 		ECVF_Default);
 
+	static TAutoConsoleVariable<float> CVarRivermaxOutputTROOverride(
+		TEXT("Rivermax.Output.TRO"), 0,
+		TEXT("If not 0, overrides transmit offset calculation (TRO) based on  frame rate and resolution with a fixed value. Value in seconds."),
+		ECVF_Default);
+
 	bool FindPayloadSize(const FRivermaxOutputStreamOptions& InOptions, uint32 InBytesPerLine, const FVideoFormatInfo& FormatInfo, uint16& OutPayloadSize)
 	{
 		using namespace UE::RivermaxCore::Private::Utils;
@@ -900,6 +905,7 @@ namespace UE::RivermaxCore::Private
 
 	void FRivermaxOutputStream::GetNextChunk()
 	{
+		bool bHasAddedTrace = false;
 		rmax_status_t Status;
 		do
 		{
@@ -922,6 +928,12 @@ namespace UE::RivermaxCore::Private
 			{
 				//We should not be here
 				Stats.ChunkRetries++;
+				if (!bHasAddedTrace)
+				{
+					UE_LOG(LogRivermax, Verbose, TEXT("No free chunks to get for chunk '%u'. Waiting"), CurrentFrame->ChunkNumber);
+					TRACE_CPUPROFILER_EVENT_SCOPE(GetNextChunk::NoFreeChunk);
+					bHasAddedTrace = true;
+				}
 			}
 			else
 			{
@@ -949,6 +961,8 @@ namespace UE::RivermaxCore::Private
 	void FRivermaxOutputStream::CommitNextChunks()
 	{
 		rmax_status_t Status;
+		int32 ErrorCount = 0;
+		const uint64 CurrentTimeNanosec = RivermaxModule->GetRivermaxManager()->GetTime();
 		do
 		{
 			//Only first chunk gets scheduled with a timestamp. Following chunks are queued after it using 0
@@ -956,7 +970,6 @@ namespace UE::RivermaxCore::Private
 			const rmax_commit_flags_t CommitFlags{};
 			if (ScheduleTime != 0)
 			{
-				const uint64 CurrentTimeNanosec = RivermaxModule->GetRivermaxManager()->GetTime();
 				TRACE_BOOKMARK(TEXT("Sched A: %llu, C: %llu"), StreamData.NextAlignmentPointNanosec, CurrentTimeNanosec);
 				if (ScheduleTime <= CurrentTimeNanosec)
 				{
@@ -975,8 +988,9 @@ namespace UE::RivermaxCore::Private
 			{
 				Stats.CommitRetries++;
 				TRACE_CPUPROFILER_EVENT_SCOPE(CommitNextChunks::QUEUEFULL);
+				++ErrorCount;
 			}
-			else if(Status == RMAX_ERR_HW_COMPLETION_ISSUE)
+			else if (Status == RMAX_ERR_HW_COMPLETION_ISSUE)
 			{
 				UE_LOG(LogRivermax, Error, TEXT("Completion issue while trying to commit next round of chunks."));
 				Listener->OnStreamError();
@@ -990,6 +1004,18 @@ namespace UE::RivermaxCore::Private
 			}
 
 		} while (Status != RMAX_OK && bIsActive);
+
+		if (bIsActive && CurrentFrame->ChunkNumber == 0)
+		{
+			UE_LOG(LogRivermax, Verbose, TEXT("Committed frame [%u]. Scheduled for '%llu'. Aligned with '%llu'. Current time '%llu'. Was late: %d. Slack: %llu. Errorcount: %d")
+				, CurrentFrame->FrameIdentifier
+				, StreamData.NextScheduleTimeNanosec
+				, StreamData.NextAlignmentPointNanosec
+				, CurrentTimeNanosec
+				, CurrentTimeNanosec >= StreamData.NextScheduleTimeNanosec ? 1 : 0
+				, StreamData.NextScheduleTimeNanosec >= CurrentTimeNanosec ? StreamData.NextScheduleTimeNanosec - CurrentTimeNanosec : 0
+				, ErrorCount);
+		}
 	}
 
 	bool FRivermaxOutputStream::Init()
@@ -1098,6 +1124,7 @@ namespace UE::RivermaxCore::Private
 
 				// Since we want to resend last frame, we need to fast forward chunk pointer to re-point to the one we just sent
 				rmax_status_t Status;
+				bool bHasAddedTrace = false;
 				do
 				{
 					Status = rmax_out_skip_chunks(StreamId, StreamMemory.ChunksPerFrameField * (Options.NumberOfBuffers - 1));
@@ -1106,7 +1133,12 @@ namespace UE::RivermaxCore::Private
 						if (Status == RMAX_ERR_NO_FREE_CHUNK)
 						{
 							// Wait until there are enough free chunk to be skipped
-							UE_LOG(LogRivermax, Warning, TEXT("No chunks ready to skip. Waiting"));
+							if (!bHasAddedTrace)
+							{
+								UE_LOG(LogRivermax, Warning, TEXT("No chunks ready to skip. Waiting"));
+								TRACE_CPUPROFILER_EVENT_SCOPE(PrepareNextFrame::NoFreeChunk);
+								bHasAddedTrace = true;
+							}
 						}
 						else
 						{
@@ -1127,6 +1159,13 @@ namespace UE::RivermaxCore::Private
 	{
 		using namespace UE::RivermaxCore::Private::Utils;
 
+		const double TROOverride = CVarRivermaxOutputTROOverride.GetValueOnAnyThread();
+		if (TROOverride != 0)
+		{
+			TransmitOffsetNanosec = TROOverride * 1E9;
+			return;
+		}
+
 		double FrameIntervalNs = StreamData.FrameFieldTimeIntervalNs;
 		const bool bIsProgressive = true;//todo MediaConfiguration.IsProgressive() 
 		uint32 PacketsInFrameField = StreamMemory.PacketsPerFrame;
@@ -1205,7 +1244,7 @@ namespace UE::RivermaxCore::Private
 			if (CurrentTime - LastStatsShownTimestamp > CVarRivermaxOutputShowStatsInterval.GetValueOnAnyThread())
 			{
 				LastStatsShownTimestamp = CurrentTime;
-				UE_LOG(LogRivermax, Log, TEXT("Stats: FrameSent: %d. CommitImmediate: %d. CommitRetries: %d. ChunkRetries: %d"), Stats.MemoryBlockSentCounter, Stats.CommitImmediate, Stats.CommitRetries, Stats.ChunkRetries);
+				UE_LOG(LogRivermax, Log, TEXT("Stats: FrameSent: %u. CommitImmediate: %u. CommitRetries: %u. ChunkRetries: %u. ChunkSkippingRetries: %u"), Stats.MemoryBlockSentCounter, Stats.CommitImmediate, Stats.CommitRetries, Stats.ChunkRetries, Stats.ChunkSkippingRetries);
 			}
 		}
 	}
@@ -1593,10 +1632,16 @@ namespace UE::RivermaxCore::Private
 					if (DeltaFrames >= 1)
 					{
 						UE_LOG(LogRivermax, Warning, TEXT("Output missed %llu frames."), DeltaFrames);
-						// For now, schedule for the following frame as normal but might need to revisit this behavior if it causes issues.
+						
+						// If we missed a sync point, this means that last scheduled frame might still be ongoing and 
+						// sending it might be crossing the frame boundary so we skip one entire frame to empty the queue.
+						NextFrameNumber = CurrentFrameNumber + 2;
+					}
+					else
+					{
+						NextFrameNumber = CurrentFrameNumber + 1;
 					}
 
-					NextFrameNumber = CurrentFrameNumber + 1;
 				}
 				else
 				{
diff --git a/Engine/Plugins/VirtualProduction/Rivermax/RivermaxCore/Source/RivermaxCore/Private/RivermaxOutputStream.h b/Engine/Plugins/VirtualProduction/Rivermax/RivermaxCore/Source/RivermaxCore/Private/RivermaxOutputStream.h
--- a/Engine/Plugins/VirtualProduction/Rivermax/RivermaxCore/Source/RivermaxCore/Private/RivermaxOutputStream.h
+++ b/Engine/Plugins/VirtualProduction/Rivermax/RivermaxCore/Source/RivermaxCore/Private/RivermaxOutputStream.h
@@ -51,6 +51,7 @@ namespace UE::RivermaxCore::Private
 	struct FRivermaxOutputStreamStats
 	{
 		uint32 ChunkRetries = 0;
+		uint32 ChunkSkippingRetries = 0;
 		uint32 TotalStrides = 0;
 		uint32 ChunkWait = 0;
 		uint32 CommitWaits = 0;
