- Adding new framelocking mode to rivermax capture.
- Reworked frame management and allocators
- Changed RTP timestamp to reduce media latency seen by analyzers

diff --git a/Engine/Plugins/VirtualProduction/Rivermax/RivermaxCore/Source/RivermaxCore/Private/RivermaxBoundaryMonitor.cpp b/Engine/Plugins/VirtualProduction/Rivermax/RivermaxCore/Source/RivermaxCore/Private/RivermaxBoundaryMonitor.cpp
new file mode 100644
--- /dev/null
+++ b/Engine/Plugins/VirtualProduction/Rivermax/RivermaxCore/Source/RivermaxCore/Private/RivermaxBoundaryMonitor.cpp
@@ -0,0 +1,215 @@
+// Copyright Epic Games, Inc. All Rights Reserved.
+
+#include "RivermaxBoundaryMonitor.h"
+
+#include "HAL/IConsoleManager.h"
+#include "HAL/RunnableThread.h"
+#include "IRivermaxCoreModule.h"
+#include "IRivermaxManager.h"
+#include "RivermaxPTPUtils.h"
+
+namespace UE::RivermaxCore
+{
+	static TMap<FFrameRate, FGuid> GBoundaryListeners;
+
+	/** 
+	 * Parses frame rate from a string 
+	 * Can either be numerator only, assuming denominator of 1 (e.g. 24)
+	 * Or can be explicit numerator over denominator (e.g. 24000/1001)
+	 */
+	static FFrameRate ParseFrameRate(const FString& Command)
+	{
+		int32 Numerator = 1;
+		int32 Denominator = 1;
+		int32 DividerIndex = INDEX_NONE;
+		Command.FindChar('/', DividerIndex);
+		if (DividerIndex != INDEX_NONE)
+		{
+			Numerator = FCString::Atoi(*Command.Left(DividerIndex));
+			Denominator = FCString::Atoi(*Command.Right(DividerIndex));
+		}
+		else
+		{
+			Numerator = FCString::Atoi(*Command);
+		}
+
+		return FFrameRate(Numerator, Denominator);
+	}
+
+	static FAutoConsoleCommand CRivermaxBoundaryMonitorAddFrameRate(
+		TEXT("Rivermax.Monitor.Add"),
+		TEXT("Add monitoring for a certain frame rate.")
+		TEXT("Example value format: \"Rivermax.Monitor.Add 24\"")
+		TEXT("Example value format: \"Rivermax.Monitor.Add 60000/1001\""),
+		FConsoleCommandWithArgsDelegate::CreateLambda([](const TArray<FString>& Args)
+			{
+				if (Args.Num() > 0)
+				{
+					const FFrameRate DesiredRate = ParseFrameRate(Args[0]);
+					if (DesiredRate.IsValid() && DesiredRate.Numerator > 0)
+					{
+						if (!GBoundaryListeners.Contains(DesiredRate))
+						{
+							IRivermaxCoreModule& RivermaxModule = FModuleManager::GetModuleChecked<IRivermaxCoreModule>("RivermaxCore");
+							GBoundaryListeners.FindOrAdd(DesiredRate) = RivermaxModule.GetRivermaxBoundaryMonitor().StartMonitoring(DesiredRate);
+						}
+					}
+				}
+			}),
+		ECVF_Default);
+
+	static FAutoConsoleCommand CRivermaxBoundaryMonitorRemoveFrameRate(
+		TEXT("Rivermax.Monitor.Remove"),
+		TEXT("Add monitoring for a certain frame rate.")
+		TEXT("Example value format: \"Rivermax.Monitor.Remove 24\"")
+		TEXT("Example value format: \"Rivermax.Monitor.Remove 60000/1001\""),
+		FConsoleCommandWithArgsDelegate::CreateLambda([](const TArray<FString>& Args)
+			{
+				if (Args.Num() > 0)
+				{
+					const FFrameRate DesiredRate = ParseFrameRate(Args[0]);
+					if (DesiredRate.IsValid() && DesiredRate.Numerator > 0)
+					{
+						if (GBoundaryListeners.Contains(DesiredRate))
+						{
+							IRivermaxCoreModule& RivermaxModule = FModuleManager::GetModuleChecked<IRivermaxCoreModule>("RivermaxCore");
+							RivermaxModule.GetRivermaxBoundaryMonitor().StopMonitoring(GBoundaryListeners[DesiredRate], DesiredRate);
+							GBoundaryListeners.Remove(DesiredRate);
+						}
+					}
+				}
+			}),
+		ECVF_Default);
+
+	FBoundaryMonitor::FBoundaryMonitor(const FFrameRate& InRate)
+		: Rate(InRate)
+	{
+		IRivermaxCoreModule& RivermaxModule = FModuleManager::GetModuleChecked<IRivermaxCoreModule>("RivermaxCore");
+		RivermaxManager = RivermaxModule.GetRivermaxManager();
+
+		bIsEnabled = true;
+
+		const FString ThreadName = FString::Printf(TEXT("Rmax_Boundary_%d/%d"), Rate.Numerator, Rate.Denominator);
+		constexpr uint32 StackSize = 32 * 1024;
+		WorkingThread.Reset(FRunnableThread::Create(this, *ThreadName, StackSize, TPri_AboveNormal));
+	}
+
+	FBoundaryMonitor::~FBoundaryMonitor()
+	{
+		bIsEnabled = false;
+		WorkingThread->WaitForCompletion();
+	}
+
+	uint32 FBoundaryMonitor::Run()
+	{
+		uint64 LastFrameNumber = UE::RivermaxCore::GetFrameNumber(RivermaxManager->GetTime(), Rate);
+		while (bIsEnabled)
+		{
+			const uint64 NextFrameNumber = LastFrameNumber + 1;
+			const uint64 TargetTimeNanosec = UE::RivermaxCore::GetAlignmentPointFromFrameNumber(NextFrameNumber, Rate);
+			const uint64 CurrentPTPTime = RivermaxManager->GetTime();
+			const double CurrentPlatformTime = FPlatformTime::Seconds();
+
+			double TimeLeftNanosec = 0.0;
+			if (TargetTimeNanosec >= CurrentPTPTime)
+			{
+				TimeLeftNanosec = TargetTimeNanosec - CurrentPTPTime;
+			}
+
+			const double WaitTimeSec = FMath::Clamp(TimeLeftNanosec / 1E9, 0.0, 5.0);
+			const double StartTime = FPlatformTime::Seconds();
+			double ActualWaitTime = 0.0;
+			{
+				static constexpr float SleepThresholdSec = 5.0f * (1.0f / 1000.0f);
+				static constexpr float SpinningTimeSec = 2.0f * (1.0f / 1000.0f);
+				if (WaitTimeSec > SleepThresholdSec)
+				{
+					FPlatformProcess::SleepNoStats(WaitTimeSec - SpinningTimeSec);
+				}
+
+				while (FPlatformTime::Seconds() < (CurrentPlatformTime + WaitTimeSec))
+				{
+					FPlatformProcess::SleepNoStats(0.f);
+				}
+			}
+
+			TRACE_BOOKMARK(TEXT("%d/%d: Frame %llu"), Rate.Numerator, Rate.Denominator, NextFrameNumber);
+			LastFrameNumber = NextFrameNumber;
+		}
+
+		return 0;
+	}
+
+FRivermaxBoundaryMonitor::~FRivermaxBoundaryMonitor()
+{
+	constexpr bool bEnableMonitoring = false;
+	EnableMonitoring(bEnableMonitoring);
+}
+
+void FRivermaxBoundaryMonitor::EnableMonitoring(bool bEnable)
+{
+	IRivermaxCoreModule* RivermaxModule = FModuleManager::GetModulePtr<IRivermaxCoreModule>("RivermaxCore");
+	if(RivermaxModule && RivermaxModule->GetRivermaxManager()->IsLibraryInitialized())
+	{
+		if (bEnable != bIsEnabled)
+		{
+			bIsEnabled = bEnable;
+
+			for (const TTuple<FFrameRate, TArray<FGuid>>& Listener : ListenersMap)
+			{
+				if (bIsEnabled)
+				{
+					StartMonitor(Listener.Key);
+				}
+				else
+				{
+					StopMonitor(Listener.Key);
+				}
+			}
+		}
+	}
+}
+
+FGuid FRivermaxBoundaryMonitor::StartMonitoring(const FFrameRate& FrameRate)
+{
+	const FGuid NewListener = FGuid::NewGuid();
+	ListenersMap.FindOrAdd(FrameRate).Add(NewListener);
+
+	if (bIsEnabled)
+	{
+		StartMonitor(FrameRate);
+	}
+
+	return NewListener;
+}
+
+void FRivermaxBoundaryMonitor::StopMonitoring(const FGuid& Requester, const FFrameRate& FrameRate)
+{
+	bool bNoMoreListeners = false;
+	if(TArray<FGuid>* Listeners = ListenersMap.Find(FrameRate))
+	{
+		Listeners->RemoveSingleSwap(Requester);
+		bNoMoreListeners = Listeners->Num() <= 0;
+	}
+
+	if (bNoMoreListeners)
+	{
+		ListenersMap.Remove(FrameRate);
+		StopMonitor(FrameRate);
+	}
+}
+
+void FRivermaxBoundaryMonitor::StartMonitor(const FFrameRate& FrameRate)
+{
+	if (!MonitorMap.Contains(FrameRate))
+	{
+		MonitorMap.Add(FrameRate, MakeUnique<FBoundaryMonitor>(FrameRate));
+	}
+}
+
+void FRivermaxBoundaryMonitor::StopMonitor(const FFrameRate& FrameRate)
+{
+	MonitorMap.Remove(FrameRate);
+}
+
+}
diff --git a/Engine/Plugins/VirtualProduction/Rivermax/RivermaxCore/Source/RivermaxCore/Private/RivermaxBoundaryMonitor.h b/Engine/Plugins/VirtualProduction/Rivermax/RivermaxCore/Source/RivermaxCore/Private/RivermaxBoundaryMonitor.h
new file mode 100644
--- /dev/null
+++ b/Engine/Plugins/VirtualProduction/Rivermax/RivermaxCore/Source/RivermaxCore/Private/RivermaxBoundaryMonitor.h
@@ -0,0 +1,78 @@
+// Copyright Epic Games, Inc. All Rights Reserved.
+
+#pragma once
+
+#include "IRivermaxBoundaryMonitor.h"
+
+#include "HAL/Runnable.h"
+#include "Misc/FrameRate.h"
+#include "Misc/Guid.h"
+
+namespace UE::RivermaxCore
+{
+
+class IRivermaxManager;
+
+/** Monitor running its own thread to add trace bookmarks at frame boundary for a given frame rate */
+class FBoundaryMonitor : public FRunnable
+{
+public:
+	FBoundaryMonitor(const FFrameRate& Rate);
+	virtual ~FBoundaryMonitor();
+
+	//~ Begin FRunnable interface
+	virtual uint32 Run() override;
+	//~ End FRunnable interface
+
+protected:
+
+	/** Thread responsible to watch for frame boundaries */
+	TUniquePtr<FRunnableThread> WorkingThread;
+
+	/** Frame rate at which we add markups */
+	FFrameRate Rate;
+
+	/** Cached pointer to rivermax manager */
+	TSharedPtr<IRivermaxManager> RivermaxManager;
+
+	/** Whether thread should be running */
+	std::atomic<bool> bIsEnabled = false;
+};
+
+/**
+ * Monitoring manager of frame boundary based on ST-2059 / PTP formulae
+ * Used to add trace bookmarks at each boundary for any frame rate monitored
+ */
+class FRivermaxBoundaryMonitor : public IRivermaxBoundaryMonitor
+{
+public:
+	virtual ~FRivermaxBoundaryMonitor();
+
+public:
+	//~ Begin IRivermaxBoundaryMonitor interface
+	virtual void EnableMonitoring(bool bEnable) override;
+	virtual FGuid StartMonitoring(const FFrameRate& FrameRate) override;
+	virtual void StopMonitoring(const FGuid& Requester, const FFrameRate& FrameRate) override;
+	//~ End IRivermaxBoundaryMonitor interface
+
+protected:
+
+	/** Starts monitoring a given frame rate */
+	void StartMonitor(const FFrameRate& FrameRate);
+	
+	/** Stops monitoring a given frame rate */
+	void StopMonitor(const FFrameRate& FrameRate);
+
+private:
+
+	/** Whether frame boundary monitor is enabled */
+	bool bIsEnabled = false;
+
+	/** Maps of monitor to a given frame rate */
+	TMap<FFrameRate, TUniquePtr<FBoundaryMonitor>> MonitorMap;
+
+	/** Maps from a given frame rate to listeners ID */
+	TMap<FFrameRate, TArray<FGuid>> ListenersMap;
+};
+
+}
diff --git a/Engine/Plugins/VirtualProduction/Rivermax/RivermaxCore/Source/RivermaxCore/Private/RivermaxCoreModule.cpp b/Engine/Plugins/VirtualProduction/Rivermax/RivermaxCore/Source/RivermaxCore/Private/RivermaxCoreModule.cpp
--- a/Engine/Plugins/VirtualProduction/Rivermax/RivermaxCore/Source/RivermaxCore/Private/RivermaxCoreModule.cpp
+++ b/Engine/Plugins/VirtualProduction/Rivermax/RivermaxCore/Source/RivermaxCore/Private/RivermaxCoreModule.cpp
@@ -3,19 +3,27 @@
 #include "RivermaxCoreModule.h"
 
 #include "Modules/ModuleManager.h"
+#include "RivermaxBoundaryMonitor.h"
 #include "RivermaxInputStream.h"
 #include "RivermaxManager.h"
 #include "RivermaxOutputStream.h"
 
 
+static TAutoConsoleVariable<int32> CVarRivermaxEnableBoundaryMonitor(
+	TEXT("Rivermax.Monitor.Enable"), 1,
+	TEXT("Whether frame boundary monitor is enabled."),
+	ECVF_Default);
+
 void FRivermaxCoreModule::StartupModule()
 {
 	RivermaxManager = MakeShared<UE::RivermaxCore::Private::FRivermaxManager>();
+	BoundaryMonitor = MakeUnique<UE::RivermaxCore::FRivermaxBoundaryMonitor>();
+	RivermaxManager->OnPostRivermaxManagerInit().AddRaw(this, &FRivermaxCoreModule::OnRivermaxManagerInitialized);
 }
 
 void FRivermaxCoreModule::ShutdownModule()
 {
-
+	BoundaryMonitor->EnableMonitoring(false);
 }
 
 TUniquePtr<UE::RivermaxCore::IRivermaxInputStream> FRivermaxCoreModule::CreateInputStream()
@@ -35,4 +43,25 @@ TSharedPtr<UE::RivermaxCore::IRivermaxManager> FRivermaxCoreModule::GetRivermaxM
 	return RivermaxManager;
 }
 
+UE::RivermaxCore::IRivermaxBoundaryMonitor& FRivermaxCoreModule::GetRivermaxBoundaryMonitor()
+{
+	return *BoundaryMonitor;
+}
+
+void FRivermaxCoreModule::OnRivermaxManagerInitialized()
+{
+	FConsoleVariableDelegate OnBoundaryMonitorEnableValueChanged = FConsoleVariableDelegate::CreateLambda([](IConsoleVariable* CVar)
+	{
+		IRivermaxCoreModule* RivermaxModule = FModuleManager::GetModulePtr<IRivermaxCoreModule>("RivermaxCore");
+		if (RivermaxModule)
+		{
+			RivermaxModule->GetRivermaxBoundaryMonitor().EnableMonitoring(CVar->GetBool());
+		}
+	});
+	CVarRivermaxEnableBoundaryMonitor.AsVariable()->SetOnChangedCallback(OnBoundaryMonitorEnableValueChanged);
+	
+	const bool bDoEnable = CVarRivermaxEnableBoundaryMonitor.GetValueOnGameThread() == 1;
+	BoundaryMonitor->EnableMonitoring(bDoEnable);
+}
+
 IMPLEMENT_MODULE(FRivermaxCoreModule, RivermaxCore);
diff --git a/Engine/Plugins/VirtualProduction/Rivermax/RivermaxCore/Source/RivermaxCore/Private/RivermaxCoreModule.h b/Engine/Plugins/VirtualProduction/Rivermax/RivermaxCore/Source/RivermaxCore/Private/RivermaxCoreModule.h
--- a/Engine/Plugins/VirtualProduction/Rivermax/RivermaxCore/Source/RivermaxCore/Private/RivermaxCoreModule.h
+++ b/Engine/Plugins/VirtualProduction/Rivermax/RivermaxCore/Source/RivermaxCore/Private/RivermaxCoreModule.h
@@ -26,10 +26,20 @@ public:
 	virtual TUniquePtr<UE::RivermaxCore::IRivermaxInputStream> CreateInputStream() override;
 	virtual TUniquePtr<UE::RivermaxCore::IRivermaxOutputStream> CreateOutputStream() override;
 	virtual TSharedPtr<UE::RivermaxCore::IRivermaxManager> GetRivermaxManager() override;
+	virtual UE::RivermaxCore::IRivermaxBoundaryMonitor& GetRivermaxBoundaryMonitor() override;
 	//~ End IRivermaxCoreModule interface
 
 private:
 
+	/** Called when manager has finished initializing */
+	void OnRivermaxManagerInitialized();
+
+private:
+
+	/** Manager handling library initialization */
 	TSharedPtr<UE::RivermaxCore::IRivermaxManager> RivermaxManager;
+
+	/** Boundary monitor used to add trace information for frame boundaries at different frame rates */
+	TUniquePtr<UE::RivermaxCore::IRivermaxBoundaryMonitor> BoundaryMonitor;
 };
 
diff --git a/Engine/Plugins/VirtualProduction/Rivermax/RivermaxCore/Source/RivermaxCore/Private/RivermaxFrameAllocator.cpp b/Engine/Plugins/VirtualProduction/Rivermax/RivermaxCore/Source/RivermaxCore/Private/RivermaxFrameAllocator.cpp
new file mode 100644
--- /dev/null
+++ b/Engine/Plugins/VirtualProduction/Rivermax/RivermaxCore/Source/RivermaxCore/Private/RivermaxFrameAllocator.cpp
@@ -0,0 +1,417 @@
+// Copyright Epic Games, Inc. All Rights Reserved.
+
+#include "RivermaxFrameAllocator.h"
+
+#include "CudaModule.h"
+#include "RivermaxLog.h"
+#include "ID3D12DynamicRHI.h"
+#include "IRivermaxOutputStream.h"
+
+namespace UE::RivermaxCore::Private
+{
+	bool FGPUAllocator::Allocate(int32 FrameCount)
+	{
+		// Allocate a single memory space that will contain all frame buffers
+		TRACE_CPUPROFILER_EVENT_SCOPE(Rmax::GPUAllocation);
+
+		FCUDAModule& CudaModule = FModuleManager::GetModuleChecked<FCUDAModule>("CUDA");
+		CudaModule.DriverAPI()->cuCtxPushCurrent(CudaModule.GetCudaContext());
+
+		// Todo: Add support for mgpu. For now, this will not work unless the memcpy does implicitely a cross gpu transfer.
+		const int GPUIndex = CudaModule.GetCudaDeviceIndex();
+		CUdevice CudaDevice;
+		CUresult Status = CudaModule.DriverAPI()->cuDeviceGet(&CudaDevice, GPUIndex);
+		if (Status != CUDA_SUCCESS)
+		{
+			UE_LOG(LogRivermax, Warning, TEXT("Can't allocate GPUMemory. Failed to get a Cuda device for GPU %d. Status: %d"), GPUIndex, Status);
+			return false;
+		}
+
+		CUmemAllocationProp AllocationProperties = {};
+		AllocationProperties.type = CU_MEM_ALLOCATION_TYPE_PINNED;
+		AllocationProperties.allocFlags.gpuDirectRDMACapable = 1;
+		AllocationProperties.allocFlags.usage = 0;
+		AllocationProperties.location.type = CU_MEM_LOCATION_TYPE_DEVICE;
+		AllocationProperties.location.id = CudaDevice;
+
+		// Get memory granularity required for cuda device. We need to align allocation with this.
+		size_t Granularity;
+		Status = CudaModule.DriverAPI()->cuMemGetAllocationGranularity(&Granularity, &AllocationProperties, CU_MEM_ALLOC_GRANULARITY_RECOMMENDED);
+		if (Status != CUDA_SUCCESS)
+		{
+			UE_LOG(LogRivermax, Warning, TEXT("Can't allocate GPUMemory. Failed to get allocation granularity. Status: %d"), Status);
+			return false;
+		}
+
+		// Cuda requires allocated memory to be aligned with a certain granularity
+		// We align each frame size to the desired granularity and multiply that by number of buffer
+		// This causes more memory to be allocated but doing a single allocation fails rmax stream creation
+		const size_t CudaAlignedFrameSize = (DesiredFrameSize % Granularity) ? DesiredFrameSize + (Granularity - (DesiredFrameSize % Granularity)) : DesiredFrameSize;
+		const size_t TotalCudaAllocSize = CudaAlignedFrameSize * FrameCount;
+
+		// Reserve contiguous memory to contain required number of buffers. 
+		CUdeviceptr CudaBaseAddress;
+		constexpr CUdeviceptr InitialAddress = 0;
+		constexpr int32 Flags = 0;
+		Status = CudaModule.DriverAPI()->cuMemAddressReserve(&CudaBaseAddress, TotalCudaAllocSize, Granularity, InitialAddress, Flags);
+		if (Status != CUDA_SUCCESS)
+		{
+			UE_LOG(LogRivermax, Warning, TEXT("Can't allocate GPUMemory. Failed to reserve memory for %d bytes. Status: %d"), TotalCudaAllocSize, Status);
+			return false;
+		}
+
+		// Make the allocation on device memory
+		CUmemGenericAllocationHandle Handle;
+		Status = CudaModule.DriverAPI()->cuMemCreate(&Handle, TotalCudaAllocSize, &AllocationProperties, Flags);
+		if (Status != CUDA_SUCCESS)
+		{
+			UE_LOG(LogRivermax, Warning, TEXT("Can't allocate GPUMemory. Failed to create memory on device. Status: %d"), Status);
+			return false;
+		}
+
+		bool bExit = false;
+		constexpr int32 Offset = 0;
+		Status = CudaModule.DriverAPI()->cuMemMap(CudaBaseAddress, TotalCudaAllocSize, Offset, Handle, Flags);
+		if (Status != CUDA_SUCCESS)
+		{
+			UE_LOG(LogRivermax, Warning, TEXT("Can't allocate GPUMemory. Failed to map memory. Status: %d"), Status);
+			// Need to release handle no matter what
+			bExit = true;
+		}
+
+		// Cache to know we need to unmap/deallocate even if it fails down the road
+		CudaAllocatedMemory = TotalCudaAllocSize;
+		CudaAllocatedMemoryBaseAddress = reinterpret_cast<void*>(CudaBaseAddress);
+
+		Status = CudaModule.DriverAPI()->cuMemRelease(Handle);
+		if (Status != CUDA_SUCCESS)
+		{
+			UE_LOG(LogRivermax, Warning, TEXT("Can't allocate GPUMemory. Failed to release handle. Status: %d"), Status);
+			return false;
+		}
+
+		if (bExit)
+		{
+			return false;
+		}
+
+		// Setup access description.
+		CUmemAccessDesc MemoryAccessDescription = {};
+		MemoryAccessDescription.flags = CU_MEM_ACCESS_FLAGS_PROT_READWRITE;
+		MemoryAccessDescription.location.type = CU_MEM_LOCATION_TYPE_DEVICE;
+		MemoryAccessDescription.location.id = CudaDevice;
+		constexpr int32 Count = 1;
+		Status = CudaModule.DriverAPI()->cuMemSetAccess(CudaBaseAddress, TotalCudaAllocSize, &MemoryAccessDescription, Count);
+		if (Status != CUDA_SUCCESS)
+		{
+			UE_LOG(LogRivermax, Warning, TEXT("Can't initialize output to use GPUDirect. Failed to configure memory access. Status: %d"), Status);
+			return false;
+		}
+
+		CUstream CudaStream;
+		Status = CudaModule.DriverAPI()->cuStreamCreate(&CudaStream, CU_STREAM_NON_BLOCKING);
+		if (Status != CUDA_SUCCESS)
+		{
+			UE_LOG(LogRivermax, Warning, TEXT("Can't initialize output to use GPUDirect. Failed to create its stream. Status: %d"), Status);
+			return false;
+		}
+
+		GPUStream = CudaStream;
+
+		Status = CudaModule.DriverAPI()->cuCtxSynchronize();
+		if (Status != CUDA_SUCCESS)
+		{
+			UE_LOG(LogRivermax, Warning, TEXT("Can't initialize output to use GPUDirect. Failed to synchronize context. Status: %d"), Status);
+			return false;
+		}
+
+		AllocatedFrameCount = FrameCount;
+		AllocatedFrameSize = CudaAlignedFrameSize;
+
+		for (int32 Index = 0; Index < AllocatedFrameCount; ++Index)
+		{
+			TSharedPtr<FRivermaxOutputFrame> Frame = MakeShared<FRivermaxOutputFrame>(Index);
+			Frame->VideoBuffer = GetFrameAddress(Index);
+			Frames.Add(MoveTemp(Frame));
+		}
+
+		CudaModule.DriverAPI()->cuCtxPopCurrent(nullptr);
+
+		return true;
+	}
+
+	void FGPUAllocator::Deallocate()
+	{
+		if (CudaAllocatedMemory > 0)
+		{
+			FCUDAModule& CudaModule = FModuleManager::GetModuleChecked<FCUDAModule>("CUDA");
+			CudaModule.DriverAPI()->cuCtxPushCurrent(CudaModule.GetCudaContext());
+
+			for (const TPair<FBufferRHIRef, void*>& Entry : BufferCudaMemoryMap)
+			{
+				if (Entry.Value)
+				{
+					CudaModule.DriverAPI()->cuMemFree(reinterpret_cast<CUdeviceptr>(Entry.Value));
+				}
+			}
+			BufferCudaMemoryMap.Empty();
+
+			const CUdeviceptr BaseAddress = reinterpret_cast<CUdeviceptr>(CudaAllocatedMemoryBaseAddress);
+			CUresult Status = CudaModule.DriverAPI()->cuMemUnmap(BaseAddress, CudaAllocatedMemory);
+			if (Status != CUDA_SUCCESS)
+			{
+				UE_LOG(LogRivermax, Warning, TEXT("Failed to unmap cuda memory. Status: %d"), Status);
+			}
+
+			Status = CudaModule.DriverAPI()->cuMemAddressFree(BaseAddress, CudaAllocatedMemory);
+			if (Status != CUDA_SUCCESS)
+			{
+				UE_LOG(LogRivermax, Warning, TEXT("Failed to free cuda memory. Status: %d"), Status);
+			}
+
+			Status = CudaModule.DriverAPI()->cuStreamDestroy(reinterpret_cast<CUstream>(GPUStream));
+			if (Status != CUDA_SUCCESS)
+			{
+				UE_LOG(LogRivermax, Warning, TEXT("Failed to destroy cuda stream. Status: %d"), Status);
+			}
+			GPUStream = nullptr;
+			AllocatedFrameCount = 0;
+			Frames.Empty();
+
+			CudaModule.DriverAPI()->cuCtxPopCurrent(nullptr);
+		}
+	}
+
+	void* FGPUAllocator::GetFrameAddress(int32 FrameIndex) const
+	{
+		if (FrameIndex >= 0 && FrameIndex < AllocatedFrameCount)
+		{
+			return reinterpret_cast<uint8*>(CudaAllocatedMemoryBaseAddress) + (FrameIndex * AllocatedFrameSize);
+		}
+
+		return nullptr;
+	}
+
+	bool FGPUAllocator::CopyData(const FRivermaxOutputVideoFrameInfo& NewFrame, const TSharedPtr<FRivermaxOutputFrame>& DestinationFrame)
+	{
+		TRACE_CPUPROFILER_EVENT_SCOPE(Rmax::GPUCopyStart);
+		check(NewFrame.GPUBuffer);
+		{
+			FCUDAModule& CudaModule = FModuleManager::GetModuleChecked<FCUDAModule>("CUDA");
+			CUresult Result = CudaModule.DriverAPI()->cuCtxPushCurrent(CudaModule.GetCudaContext());
+
+			void* MappedPointer = GetMappedAddress(NewFrame.GPUBuffer);
+			if (MappedPointer == nullptr)
+			{
+				UE_LOG(LogRivermax, Error, TEXT("Failed to find a mapped memory address for captured buffer. Stopping capture."));
+				return false;
+			}
+
+			const CUdeviceptr CudaMemoryPointer = reinterpret_cast<CUdeviceptr>(MappedPointer);
+			Result = CudaModule.DriverAPI()->cuMemcpyDtoDAsync(reinterpret_cast<CUdeviceptr>(DestinationFrame->VideoBuffer), CudaMemoryPointer, DesiredFrameSize, reinterpret_cast<CUstream>(GPUStream));
+			if (Result != CUDA_SUCCESS)
+			{
+				UE_LOG(LogRivermax, Error, TEXT("Failed to copy captured bufer to cuda memory. Stopping capture. Error: %d"), Result);
+				return false;
+			}
+
+			// Callback called by Cuda when stream work has completed on cuda engine (MemCpy -> Callback)
+			// Once Memcpy has been done, we know we can mark that memory as available to be sent. 
+
+			auto CudaCallback = [](void* userData)
+			{
+				TRACE_CPUPROFILER_EVENT_SCOPE(Rmax::GPUCopyDone);
+
+				FGPUAllocator* Allocator = reinterpret_cast<FGPUAllocator*>(userData);
+				TOptional<TSharedPtr<FRivermaxOutputFrame>> FrameCopied = Allocator->PendingCopies.Dequeue();
+				if (FrameCopied.IsSet())
+				{
+					FrameCopied.GetValue()->bIsVideoBufferReady = true;
+					Allocator->OnDataCopiedDelegate.ExecuteIfBound(FrameCopied.GetValue());
+				}
+			};
+
+			// Add pending frame for cuda callback 
+			PendingCopies.Enqueue(DestinationFrame);
+
+			// Schedule a callback to make the frame available
+			CudaModule.DriverAPI()->cuLaunchHostFunc(reinterpret_cast<CUstream>(GPUStream), CudaCallback, this);
+
+			FCUDAModule::CUDA().cuCtxPopCurrent(nullptr);
+
+			return true;
+		}
+
+		return false;
+	}
+
+	FGPUAllocator::FGPUAllocator(FIntPoint InFrameResolution, uint32 InStride, FOnFrameDataCopiedDelegate InOnDataCopiedDelegate)
+		: Super(InFrameResolution, InStride, InOnDataCopiedDelegate)
+	{
+
+	}
+
+	void* FGPUAllocator::GetMappedAddress(const FBufferRHIRef& InBuffer)
+	{
+		// If we are here, d3d12 had to have been validated
+		const ERHIInterfaceType RHIType = RHIGetInterfaceType();
+		check(RHIType == ERHIInterfaceType::D3D12);
+
+		//Do we already have a mapped address for this buffer
+		if (BufferCudaMemoryMap.Find((InBuffer)) == nullptr)
+		{
+			int64 BufferMemorySize = 0;
+			CUexternalMemory MappedExternalMemory = nullptr;
+			HANDLE D3D12BufferHandle = 0;
+			CUDA_EXTERNAL_MEMORY_HANDLE_DESC CudaExtMemHandleDesc = {};
+
+			// Create shared handle for our buffer
+			{
+				TRACE_CPUPROFILER_EVENT_SCOPE(Rmax_D3D12CreateSharedHandle);
+
+				ID3D12Resource* NativeD3D12Resource = GetID3D12DynamicRHI()->RHIGetResource(InBuffer);
+				BufferMemorySize = GetID3D12DynamicRHI()->RHIGetResourceMemorySize(InBuffer);
+
+				TRefCountPtr<ID3D12Device> OwnerDevice;
+				HRESULT QueryResult;
+				if ((QueryResult = NativeD3D12Resource->GetDevice(IID_PPV_ARGS(OwnerDevice.GetInitReference()))) != S_OK)
+				{
+					UE_LOG(LogRivermax, Error, TEXT("Failed to get D3D12 device for captured buffer ressource: %d)"), QueryResult);
+					return nullptr;
+				}
+
+				if ((QueryResult = OwnerDevice->CreateSharedHandle(NativeD3D12Resource, NULL, GENERIC_ALL, NULL, &D3D12BufferHandle)) != S_OK)
+				{
+					UE_LOG(LogRivermax, Error, TEXT("Failed to create shared handle for captured buffer ressource: %d"), QueryResult);
+					return nullptr;
+				}
+
+				CudaExtMemHandleDesc.type = CU_EXTERNAL_MEMORY_HANDLE_TYPE_D3D12_RESOURCE;
+				CudaExtMemHandleDesc.handle.win32.name = nullptr;
+				CudaExtMemHandleDesc.handle.win32.handle = D3D12BufferHandle;
+				CudaExtMemHandleDesc.size = BufferMemorySize;
+				CudaExtMemHandleDesc.flags |= CUDA_EXTERNAL_MEMORY_DEDICATED;
+			}
+
+			{
+				TRACE_CPUPROFILER_EVENT_SCOPE(Rmax_CudaImportMemory);
+
+				const CUresult Result = FCUDAModule::CUDA().cuImportExternalMemory(&MappedExternalMemory, &CudaExtMemHandleDesc);
+
+				if (D3D12BufferHandle)
+				{
+					CloseHandle(D3D12BufferHandle);
+				}
+
+				if (Result != CUDA_SUCCESS)
+				{
+					UE_LOG(LogRivermax, Error, TEXT("Failed to import shared buffer. Error: %d"), Result);
+					return nullptr;
+				}
+			}
+
+			{
+				TRACE_CPUPROFILER_EVENT_SCOPE(Rmax_MapCudaMemory);
+
+				CUDA_EXTERNAL_MEMORY_BUFFER_DESC BufferDescription = {};
+				BufferDescription.offset = 0;
+				BufferDescription.size = BufferMemorySize;
+				CUdeviceptr NewMemory;
+				const CUresult Result = FCUDAModule::CUDA().cuExternalMemoryGetMappedBuffer(&NewMemory, MappedExternalMemory, &BufferDescription);
+				if (Result != CUDA_SUCCESS || NewMemory == 0)
+				{
+					UE_LOG(LogRivermax, Error, TEXT("Failed to get shared buffer mapped memory. Error: %d"), Result);
+					return nullptr;
+				}
+
+				BufferCudaMemoryMap.Add(InBuffer, reinterpret_cast<void*>(NewMemory));
+			}
+		}
+
+		// At this point, we have the mapped buffer in cuda space and we can use it to schedule a memcpy on cuda engine.
+		return BufferCudaMemoryMap[InBuffer];
+	}
+
+	bool FSystemAllocator::Allocate(int32 FrameCount)
+	{
+		TRACE_CPUPROFILER_EVENT_SCOPE(Rmax::SystemAllocation);
+
+		constexpr uint32 CacheLineSize = PLATFORM_CACHE_LINE_SIZE;
+
+		AllocatedFrameSize = Align(DesiredFrameSize, CacheLineSize);
+		const size_t NeededSystemMemory = Align(AllocatedFrameSize * FrameCount, CacheLineSize);
+		SystemMemoryBaseAddress = FMemory::Malloc(AllocatedFrameSize * FrameCount, CacheLineSize);
+
+		if (SystemMemoryBaseAddress)
+		{
+			AllocatedFrameCount = FrameCount;
+
+			for (int32 Index = 0; Index < AllocatedFrameCount; ++Index)
+			{
+				TSharedPtr<FRivermaxOutputFrame> Frame = MakeShared<FRivermaxOutputFrame>(Index);
+				Frame->VideoBuffer = GetFrameAddress(Index);
+				Frames.Add(MoveTemp(Frame));
+			}
+
+			return true;
+		}
+
+		return false;
+	}
+
+	void FSystemAllocator::Deallocate()
+	{
+		if (SystemMemoryBaseAddress)
+		{
+			FMemory::Free(SystemMemoryBaseAddress);
+			SystemMemoryBaseAddress = nullptr;
+			AllocatedFrameCount = 0;
+			Frames.Empty();
+		}
+	}
+
+	void* FSystemAllocator::GetFrameAddress(int32 FrameIndex) const
+	{
+		if (FrameIndex >= 0 && FrameIndex < AllocatedFrameCount)
+		{
+			return reinterpret_cast<uint8*>(SystemMemoryBaseAddress) + (FrameIndex * AllocatedFrameSize);
+		}
+
+		return nullptr;
+	}
+
+	FSystemAllocator::FSystemAllocator(FIntPoint InFrameResolution, uint32 InStride, FOnFrameDataCopiedDelegate InOnDataCopiedDelegate)
+		: Super(InFrameResolution, InStride, InOnDataCopiedDelegate)
+	{
+
+	}
+
+	bool FSystemAllocator::CopyData(const FRivermaxOutputVideoFrameInfo& NewFrame, const TSharedPtr<FRivermaxOutputFrame>& DestinationFrame)
+	{
+		FMemory::Memcpy(DestinationFrame->VideoBuffer, NewFrame.VideoBuffer, DesiredFrameSize);
+
+		DestinationFrame->bIsVideoBufferReady = true;
+		OnDataCopiedDelegate.ExecuteIfBound(DestinationFrame);
+
+		return true;
+	}
+
+	FBaseFrameAllocator::FBaseFrameAllocator(FIntPoint InFrameResolution, uint32 InStride, FOnFrameDataCopiedDelegate InOnDataCopiedDelegate)
+		: FrameResolution(InFrameResolution)
+		, Stride(InStride)
+		, DesiredFrameSize(FrameResolution.Y * Stride)
+		, OnDataCopiedDelegate(MoveTemp(InOnDataCopiedDelegate))
+	{
+
+	}
+
+	const TSharedPtr<FRivermaxOutputFrame> FBaseFrameAllocator::GetFrame(int32 Index) const
+	{
+		if (Frames.IsValidIndex(Index))
+		{
+			return Frames[Index];
+		}
+		return nullptr;
+	}
+}
+
diff --git a/Engine/Plugins/VirtualProduction/Rivermax/RivermaxCore/Source/RivermaxCore/Private/RivermaxFrameAllocator.h b/Engine/Plugins/VirtualProduction/Rivermax/RivermaxCore/Source/RivermaxCore/Private/RivermaxFrameAllocator.h
new file mode 100644
--- /dev/null
+++ b/Engine/Plugins/VirtualProduction/Rivermax/RivermaxCore/Source/RivermaxCore/Private/RivermaxFrameAllocator.h
@@ -0,0 +1,143 @@
+// Copyright Epic Games, Inc. All Rights Reserved.
+
+#pragma once
+
+#include "CoreTypes.h"
+
+#include "Containers/SpscQueue.h"
+#include "Delegates/DelegateCombinations.h"
+#include "RivermaxOutputFrame.h"
+#include "RHI.h"
+
+namespace UE::RivermaxCore
+{
+	struct FRivermaxOutputVideoFrameInfo;
+}
+
+
+namespace UE::RivermaxCore::Private
+{
+	struct FRivermaxOutputFrame;
+
+	/** Delegate triggered when data has been copied into a given frame */
+	DECLARE_DELEGATE_OneParam(FOnFrameDataCopiedDelegate, const TSharedPtr<FRivermaxOutputFrame>&);
+
+	/** Base class for frame allocation. Currently supports GPU allocation using CUDA and System memory */
+	class FBaseFrameAllocator
+	{
+	public:
+		FBaseFrameAllocator(FIntPoint FrameResolution, uint32 Stride, FOnFrameDataCopiedDelegate InOnDataCopiedDelegate);
+		virtual ~FBaseFrameAllocator() = default;
+
+		/** Allocates memory to hold FrameCount frames */
+		virtual bool Allocate(int32 FrameCount) = 0;
+
+		/** Deallocate memory */
+		virtual void Deallocate() = 0;
+
+		/** Initiate a copy into a given frame */
+		virtual bool CopyData(const FRivermaxOutputVideoFrameInfo& NewFrame, const TSharedPtr<FRivermaxOutputFrame>& DestinationFrame) = 0;
+		
+		/** Returns allocated frame for a given index */
+		const TSharedPtr<FRivermaxOutputFrame> GetFrame(int32 Index) const;
+
+	protected:
+		
+		/** Returns buffer address of a given frame */
+		virtual void* GetFrameAddress(int32 FrameIndex) const = 0;
+
+	protected:
+
+		/** List of allocated frames. */
+		TArray<TSharedPtr<FRivermaxOutputFrame>> Frames;
+
+		/** Resolution of frames to be allocated */
+		FIntPoint FrameResolution = FIntPoint::ZeroValue;
+
+		/** Stride of a line */
+		uint32 Stride = 0;
+		
+		/** Total desired size of a frame */
+		int32 DesiredFrameSize = 0;
+
+		/** Delegate callbacked when data has been copied. Depending on allocator, memcopy can be async. */
+		FOnFrameDataCopiedDelegate OnDataCopiedDelegate;
+	};
+
+	class FGPUAllocator : public FBaseFrameAllocator
+	{
+		using Super = FBaseFrameAllocator;
+
+	public:
+		FGPUAllocator(FIntPoint FrameResolution, uint32 Stride, FOnFrameDataCopiedDelegate InOnDataCopiedDelegate);
+
+		//~ Begin FBaseFrameAllocator interface
+		virtual bool Allocate(int32 FrameCount) override;
+		virtual void Deallocate() override;
+		virtual bool CopyData(const FRivermaxOutputVideoFrameInfo& NewFrame, const TSharedPtr<FRivermaxOutputFrame>& DestinationFrame) override;
+	
+	protected:
+		virtual void* GetFrameAddress(int32 FrameIndex) const override;
+		//~ End FBaseFrameAllocator interface
+
+	private:
+		/** Get mapped address in cuda space for a given buffer. Cache will be updated if not found */
+		void* GetMappedAddress(const FBufferRHIRef& InBuffer);
+
+	private:
+
+		/** Number of frames that was allocated */
+		int32 AllocatedFrameCount = 0;
+
+		/** Allocated size of a frame */
+		size_t AllocatedFrameSize = 0;
+
+		/** Allocated memory base address used when it's time to free */
+		void* CudaAllocatedMemoryBaseAddress = nullptr;
+
+		/** Total allocated gpu memory. */
+		int32 CudaAllocatedMemory = 0;
+
+		/** Cuda stream used for our operations */
+		void* GPUStream = nullptr;
+
+		/** Map between buffer we are sending and their mapped address in gpu space */
+		TMap<FBufferRHIRef, void*> BufferCudaMemoryMap;
+
+		/** Frames that are waiting for a cuda memcopy. Will be dequeued when our cuda task has completed */
+		TSpscQueue<TSharedPtr<FRivermaxOutputFrame>> PendingCopies;
+	};
+
+	class FSystemAllocator : public FBaseFrameAllocator
+	{
+		using Super = FBaseFrameAllocator;
+
+	public:
+		FSystemAllocator(FIntPoint FrameResolution, uint32 Stride, FOnFrameDataCopiedDelegate InOnDataCopiedDelegate);
+		
+		//~ Begin FBaseFrameAllocator interface
+		virtual bool Allocate(int32 FrameCount) override;
+		virtual void Deallocate() override;
+		virtual bool CopyData(const FRivermaxOutputVideoFrameInfo& NewFrame, const TSharedPtr<FRivermaxOutputFrame>& DestinationFrame) override;
+	
+	protected:
+		virtual void* GetFrameAddress(int32 FrameIndex) const override;
+		//~ End FBaseFrameAllocator interface
+	
+	private:
+
+		/** Number of frames that was allocated */
+		int32 AllocatedFrameCount = 0;
+
+		/** Size of a frame that was allocated */
+		size_t AllocatedFrameSize = 0;
+
+		/** Total memory allocated to hold all frames */
+		size_t SystemAllocatedMemory = 0;
+
+		/** Allocated memory base address used when it's time to free */
+		void* SystemMemoryBaseAddress = nullptr;
+	};
+}
+
+
diff --git a/Engine/Plugins/VirtualProduction/Rivermax/RivermaxCore/Source/RivermaxCore/Private/RivermaxFrameManager.cpp b/Engine/Plugins/VirtualProduction/Rivermax/RivermaxCore/Source/RivermaxCore/Private/RivermaxFrameManager.cpp
new file mode 100644
--- /dev/null
+++ b/Engine/Plugins/VirtualProduction/Rivermax/RivermaxCore/Source/RivermaxCore/Private/RivermaxFrameManager.cpp
@@ -0,0 +1,216 @@
+// Copyright Epic Games, Inc. All Rights Reserved.
+
+#include "RivermaxFrameManager.h"
+
+#include "IRivermaxCoreModule.h"
+#include "IRivermaxManager.h"
+#include "RivermaxLog.h"
+
+namespace UE::RivermaxCore::Private
+{
+	FFrameManager::~FFrameManager()
+	{
+		Cleanup();
+	}
+
+	EFrameMemoryLocation FFrameManager::Initialize(const FFrameManagerSetupArgs& Args)
+	{
+		RivermaxManager = FModuleManager::GetModuleChecked<IRivermaxCoreModule>(TEXT("RivermaxCore")).GetRivermaxManager();
+		check(RivermaxManager);
+
+		OnFrameReadyDelegate = Args.OnFrameReadyDelegate;
+		OnPreFrameReadyDelegate = Args.OnPreFrameReadyDelegate;
+		OnFreeFrameDelegate = Args.OnFreeFrameDelegate;
+		FrameResolution = Args.Resolution;
+		TotalFrameCount = Args.NumberOfFrames;
+		const uint32 FrameSize = FrameResolution.Y * Args.Stride;
+		FOnFrameDataCopiedDelegate OnDataCopiedDelegate = FOnFrameDataCopiedDelegate::CreateRaw(this, &FFrameManager::OnDataCopied);
+
+		if (Args.bTryGPUAllocation)
+		{
+			FrameAllocator = MakeUnique<FGPUAllocator>(FrameResolution, Args.Stride, OnDataCopiedDelegate);
+			if (FrameAllocator->Allocate(TotalFrameCount))
+			{
+				MemoryLocation = EFrameMemoryLocation::GPU;
+			}
+		}
+
+		if (MemoryLocation == EFrameMemoryLocation::None)
+		{
+			FrameAllocator = MakeUnique<FSystemAllocator>(FrameResolution, Args.Stride, OnDataCopiedDelegate);
+			if (FrameAllocator->Allocate(TotalFrameCount))
+			{
+				MemoryLocation = EFrameMemoryLocation::System;
+			}
+		}
+
+		if (MemoryLocation != EFrameMemoryLocation::None)
+		{
+			// Create frame state tracking containers
+			FreeFrames.Reserve(TotalFrameCount);
+			PendingFrames.Reserve(TotalFrameCount);
+			ReadyFrames.Reserve(TotalFrameCount);
+			for (uint32 Index = 0; Index < TotalFrameCount; ++Index)
+			{
+				// All frames default to being available
+				FreeFrames.Add(Index);
+			}
+		}
+
+		return MemoryLocation;
+	}
+
+	void FFrameManager::Cleanup()
+	{
+		if (FrameAllocator)
+		{
+			FrameAllocator->Deallocate();
+			FrameAllocator.Reset();
+		}
+	}
+
+	TSharedPtr<UE::RivermaxCore::Private::FRivermaxOutputFrame> FFrameManager::GetFreeFrame()
+	{
+		FScopeLock Lock(&ContainersCritSec);
+		if (!FreeFrames.IsEmpty())
+		{
+			return FrameAllocator->GetFrame(FreeFrames[0]);
+		}
+
+		return nullptr;
+	}
+
+	void FFrameManager::MarkAsPending(const TSharedPtr<FRivermaxOutputFrame>& Frame)
+	{
+		FScopeLock Lock(&ContainersCritSec);
+
+		ensure(FreeFrames.RemoveSingle(Frame->FrameIndex) >= 1);
+		PendingFrames.Add(Frame->FrameIndex);
+	}
+
+	TSharedPtr<UE::RivermaxCore::Private::FRivermaxOutputFrame> FFrameManager::GetNextFrame(uint32 NextFrameIdentifier)
+	{
+		FScopeLock Lock(&ContainersCritSec);
+		
+		TSharedPtr<FRivermaxOutputFrame> NextFrame = GetPendingFrame(NextFrameIdentifier);
+		if (NextFrame)
+		{
+			// We found a pending / reserved frame matching identifier
+			return NextFrame;
+		}
+
+		//Otherwise, prepare next free frame
+		if (NextFrame = GetFreeFrame())
+		{
+			NextFrame->Reset();
+			NextFrame->FrameIdentifier = NextFrameIdentifier;
+			MarkAsPending(NextFrame);
+		}
+
+		return NextFrame;
+	}
+
+	TSharedPtr<FRivermaxOutputFrame> FFrameManager::GetPendingFrame(uint32 FrameIdentifier)
+	{
+		FScopeLock Lock(&ContainersCritSec);
+		if (!PendingFrames.IsEmpty())
+		{
+			for (const uint32 FrameIndex : PendingFrames)
+			{
+				TSharedPtr<FRivermaxOutputFrame> Frame = FrameAllocator->GetFrame(FrameIndex);
+				if (Frame->FrameIdentifier == FrameIdentifier)
+				{
+					// We found a reserved / pending frame corresponding to next frame identifier
+					return Frame;
+				}
+			}
+		}
+		return nullptr;
+	}
+
+	TSharedPtr<FRivermaxOutputFrame> FFrameManager::GetReadyFrame()
+	{
+		FScopeLock Lock(&ContainersCritSec);
+		if (!ReadyFrames.IsEmpty())
+		{
+			const uint32 NextReadyFrame = ReadyFrames[0];
+			return FrameAllocator->GetFrame(NextReadyFrame);
+		}
+	
+		return nullptr;
+	}
+
+	void FFrameManager::MarkAsSent(const TSharedPtr<FRivermaxOutputFrame>& Frame)
+	{
+		{
+			FScopeLock Lock(&ContainersCritSec);
+
+			Frame->Reset();
+			FreeFrames.Add(SendingFrame);
+
+			if (ensure(Frame->FrameIndex == SendingFrame))
+			{
+				SendingFrame = INDEX_NONE;
+			}
+		}
+		
+		OnFreeFrameDelegate.ExecuteIfBound();
+	}
+
+	void FFrameManager::MarkAsReady(const TSharedPtr<FRivermaxOutputFrame>& Frame)
+	{
+		{
+			// Make frame available to be sent
+			FScopeLock Lock(&ContainersCritSec);
+			Frame->ReadyTimestamp = RivermaxManager->GetTime();
+			PendingFrames.RemoveSingle(Frame->FrameIndex);
+			ReadyFrames.Add(Frame->FrameIndex);
+		}
+
+		OnFrameReadyDelegate.ExecuteIfBound();
+	}
+
+	void FFrameManager::MarkAsSending(const TSharedPtr<FRivermaxOutputFrame>& Frame)
+	{
+		// Make frame available to be sent
+		FScopeLock Lock(&ContainersCritSec);
+		if (ensure(!ReadyFrames.IsEmpty()))
+		{
+			SendingFrame = ReadyFrames[0];
+			ensure(SendingFrame == Frame->FrameIndex);
+			ReadyFrames.RemoveAt(0);
+		}
+	}
+
+	const TSharedPtr<FRivermaxOutputFrame> FFrameManager::GetFrame(int32 Index) const
+	{
+		return FrameAllocator->GetFrame(Index);
+	}
+
+	bool FFrameManager::SetFrameData(const FRivermaxOutputVideoFrameInfo& NewFrameInfo)
+	{
+		if (TSharedPtr<FRivermaxOutputFrame> NextFrame = GetNextFrame(NewFrameInfo.FrameIdentifier))
+		{
+			return FrameAllocator->CopyData(NewFrameInfo, NextFrame);
+		}
+		return true;
+	}
+
+	void FFrameManager::OnDataCopied(const TSharedPtr<FRivermaxOutputFrame>& Frame)
+	{
+		OnPreFrameReadyDelegate.ExecuteIfBound();
+
+		if (TSharedPtr<FRivermaxOutputFrame> AvailableFrame = GetNextFrame(Frame->FrameIdentifier))
+		{
+			if (AvailableFrame->IsReadyToBeSent())
+			{
+				const FString TraceName = FString::Format(TEXT("Rmax::FrameReady {0}"), { AvailableFrame->FrameIndex });
+				TRACE_CPUPROFILER_EVENT_SCOPE_TEXT(*TraceName);
+				TRACE_CPUPROFILER_EVENT_SCOPE_TEXT(*FString::Printf(TEXT("MediaCapturePipe: %u"), AvailableFrame->FrameIdentifier));
+				
+				MarkAsReady(AvailableFrame);
+			}
+		}
+	}
+}
+
diff --git a/Engine/Plugins/VirtualProduction/Rivermax/RivermaxCore/Source/RivermaxCore/Private/RivermaxFrameManager.h b/Engine/Plugins/VirtualProduction/Rivermax/RivermaxCore/Source/RivermaxCore/Private/RivermaxFrameManager.h
new file mode 100644
--- /dev/null
+++ b/Engine/Plugins/VirtualProduction/Rivermax/RivermaxCore/Source/RivermaxCore/Private/RivermaxFrameManager.h
@@ -0,0 +1,178 @@
+// Copyright Epic Games, Inc. All Rights Reserved.
+
+#pragma once
+
+#include "CoreTypes.h"
+
+#include "Containers/SpscQueue.h"
+#include "RivermaxFormats.h"
+#include "RivermaxFrameAllocator.h"
+#include "RivermaxOutputFrame.h"
+#include "RivermaxTypes.h"
+
+namespace UE::RivermaxCore
+{
+	class IRivermaxManager;
+}
+
+namespace UE::RivermaxCore::Private
+{
+	class FBaseFrameAllocator;
+
+	/** Where frame memory is allocated */
+	enum class EFrameMemoryLocation : uint8
+	{
+		/** No memory was allocated */
+		None,
+
+		/** Memory allocated in system memory */
+		System,
+
+		/** Memory allocated on GPU. Cuda space is used at the moment. */
+		GPU
+	};
+
+	DECLARE_DELEGATE(FOnFrameReadyDelegate);
+	DECLARE_DELEGATE(FOnPreFrameReadyDelegate);
+	DECLARE_DELEGATE(FOnFreeFrameDelegate);
+
+	/** Holds arguments to configure frame manager during initialization */
+	struct FFrameManagerSetupArgs
+	{
+		/** Resolution of video frames to allocate */
+		FIntPoint Resolution = FIntPoint::ZeroValue;
+
+		/** Stride of a line of video frame */
+		uint32 Stride = 0;
+
+		/** Number of video frames required */
+		uint8 NumberOfFrames = 0;
+
+		/** Whether we should try allocating on GPU */
+		bool bTryGPUAllocation = true;
+
+		/** Delegate called when a frame is now free to use */
+		FOnFreeFrameDelegate OnFreeFrameDelegate = nullptr;
+
+		/** Delegate triggered just before a frame is enqueued to be sent */
+		FOnPreFrameReadyDelegate OnPreFrameReadyDelegate = nullptr;
+		
+		/** Delegate called when a frame is now ready to be sent */
+		FOnFrameReadyDelegate OnFrameReadyDelegate = nullptr;
+	};
+
+	/** 
+	 * Class managing frames that we output over network 
+	 * Handles memory allocation and state tracking
+	 * 
+	 * States of a frame
+	 * 
+	 * Free :	Frame can be used by the capture system.
+	 * Pending:	Frame is being used by the capture system. 
+	 *			Data isn't ready to be sent out yet but it's reserved for a given identifier.
+	 * Ready:	Frame is ready to be sent. Data has been copied into it. 
+	 * Sending: Frame is being actively sent out the wire. Can't modify it until next frame boundary
+	 *
+	 * Frame rate control
+	 * 
+	 * Sending a frame out takes a full frame interval so if capture system goes faster than output rate
+	 * Free frames list will get depleted. If frame locking mode is used, getting the next free frame
+	 * will block until a new one is available which will happen at next frame boundary
+	 * Rendering and capturing the next frame might be quick but when ready to present it, it will get stalled.
+	 * This will cause the engine's frame rate to match the output frame rate.
+	 * 
+	 */
+	class FFrameManager
+	{
+	public:
+		virtual ~FFrameManager();
+
+		/** Initializes frame manager with a set of options. Returns where frames were allocated. */
+		EFrameMemoryLocation Initialize(const FFrameManagerSetupArgs& Args);
+		
+		/** Requests cleanup of allocated memory */
+		void Cleanup();
+
+		/** Returns a frame not being used */
+		TSharedPtr<FRivermaxOutputFrame> GetFreeFrame();
+
+		/**
+		 * Get next frame that can be used. 
+		 * If a pending one matching identifier is found, it is returned.
+		 * If not, a free frame will be returned, if any are available
+		 */
+		TSharedPtr<FRivermaxOutputFrame> GetNextFrame(uint32 NextFrameIdentifier);
+
+		/** Returns a pending frame with the given identifier. If none is found, returns nullptr */
+		TSharedPtr<FRivermaxOutputFrame> GetPendingFrame(uint32 FrameIdentifier);
+		
+		/** Returns next frame ready to be sent. Can be null. */
+		TSharedPtr<FRivermaxOutputFrame> GetReadyFrame();
+
+		/** Returns a frame for a given index */
+		const TSharedPtr<FRivermaxOutputFrame> GetFrame(int32 Index) const;
+
+		/** Mark a frame as being used and data should be coming soon */
+		void MarkAsPending(const TSharedPtr<FRivermaxOutputFrame>& Frame);
+		
+		/** Mark a frame as being ready to be sent */
+		void MarkAsReady(const TSharedPtr<FRivermaxOutputFrame>& Frame);
+		
+		/** Mark a frame as being sent */
+		void MarkAsSending(const TSharedPtr<FRivermaxOutputFrame>& Frame);
+
+		/** Mark a frame as sent */
+		void MarkAsSent(const TSharedPtr<FRivermaxOutputFrame>& Frame);
+
+		/** Initiates memory copy for a given frame */
+		bool SetFrameData(const FRivermaxOutputVideoFrameInfo& NewFrame);
+
+	private:
+
+		/** Called back when copy request was completed by allocator */
+		void OnDataCopied(const TSharedPtr<FRivermaxOutputFrame>& Frame);
+
+	private:
+
+		/** Resolution of video frames */
+		FIntPoint FrameResolution = FIntPoint::ZeroValue;
+
+		/** Number of frames allocated */
+		uint32 TotalFrameCount = 0;
+
+		/** Location of memory allocated */
+		EFrameMemoryLocation MemoryLocation = EFrameMemoryLocation::None;
+
+		/** Critical section to protect various frame containers */
+		FCriticalSection ContainersCritSec;
+
+		/** Frame allocator dealing with memory operation */
+		TUniquePtr<FBaseFrameAllocator> FrameAllocator;
+
+		/** Delegate triggered when a frame is free to use */
+		FOnFreeFrameDelegate OnFreeFrameDelegate = nullptr;
+
+		/** Delegate triggered just before a frame is enqueued to be sent */
+		FOnPreFrameReadyDelegate OnPreFrameReadyDelegate = nullptr;
+		
+		/** Delegate triggered when a frame is ready to be sent (video data has been copied) */
+		FOnFrameReadyDelegate OnFrameReadyDelegate = nullptr;
+		
+		/** List of available frames. */
+		TArray<int32> FreeFrames;
+		
+		/** List of frames being used but not ready to be sent yet  */
+		TArray<int32> PendingFrames;
+
+		/** List of frames ready to send */
+		TArray<int32> ReadyFrames;
+
+		/** Frame currently being sent */
+		int32 SendingFrame = 0;
+
+		/** Quick access to rivermax manager */
+		TSharedPtr<IRivermaxManager> RivermaxManager;
+	};
+}
+
+
diff --git a/Engine/Plugins/VirtualProduction/Rivermax/RivermaxCore/Source/RivermaxCore/Private/RivermaxOutputFrame.cpp b/Engine/Plugins/VirtualProduction/Rivermax/RivermaxCore/Source/RivermaxCore/Private/RivermaxOutputFrame.cpp
--- a/Engine/Plugins/VirtualProduction/Rivermax/RivermaxCore/Source/RivermaxCore/Private/RivermaxOutputFrame.cpp
+++ b/Engine/Plugins/VirtualProduction/Rivermax/RivermaxCore/Source/RivermaxCore/Private/RivermaxOutputFrame.cpp
@@ -4,22 +4,13 @@
 
 namespace UE::RivermaxCore::Private
 {
-	FRivermaxOutputFrame::FRivermaxOutputFrame(uint32 InFrameIndex, TFunction<void(void*)> InDeallocationFunction)
+	FRivermaxOutputFrame::FRivermaxOutputFrame(uint32 InFrameIndex)
 		: FrameIndex(InFrameIndex)
 		, VideoBuffer(nullptr)
-		, DeallocationFunc(InDeallocationFunction)
 	{
 		Clear();
 	}
 
-	FRivermaxOutputFrame::~FRivermaxOutputFrame()
-	{
-		if (DeallocationFunc)
-		{
-			DeallocationFunc(VideoBuffer);
-		}
-	}
-
 	bool FRivermaxOutputFrame::IsReadyToBeSent() const
 	{
 		return FrameIdentifier != InvalidIdentifier && bIsVideoBufferReady;
diff --git a/Engine/Plugins/VirtualProduction/Rivermax/RivermaxCore/Source/RivermaxCore/Private/RivermaxOutputFrame.h b/Engine/Plugins/VirtualProduction/Rivermax/RivermaxCore/Source/RivermaxCore/Private/RivermaxOutputFrame.h
--- a/Engine/Plugins/VirtualProduction/Rivermax/RivermaxCore/Source/RivermaxCore/Private/RivermaxOutputFrame.h
+++ b/Engine/Plugins/VirtualProduction/Rivermax/RivermaxCore/Source/RivermaxCore/Private/RivermaxOutputFrame.h
@@ -14,8 +14,7 @@ namespace UE::RivermaxCore::Private
 	 */
 	struct FRivermaxOutputFrame
 	{
-		FRivermaxOutputFrame(uint32 InFrameIndex, TFunction<void(void*)> InDeallocationFunction);
-		~FRivermaxOutputFrame();
+		FRivermaxOutputFrame(uint32 InFrameIndex);
 		FRivermaxOutputFrame(const FRivermaxOutputFrame&) = delete;
 		FRivermaxOutputFrame& operator=(const FRivermaxOutputFrame&) = delete;
 
@@ -64,10 +63,6 @@ namespace UE::RivermaxCore::Private
 		/** Time at which this frame was made available to be sent */
 		uint64 ReadyTimestamp = 0;
 
-	private:
-
-		/** Method called to deallocate memory. */
-		TFunction<void(void*)> DeallocationFunc;
 	};
 }
 
diff --git a/Engine/Plugins/VirtualProduction/Rivermax/RivermaxCore/Source/RivermaxCore/Private/RivermaxOutputStream.cpp b/Engine/Plugins/VirtualProduction/Rivermax/RivermaxCore/Source/RivermaxCore/Private/RivermaxOutputStream.cpp
--- a/Engine/Plugins/VirtualProduction/Rivermax/RivermaxCore/Source/RivermaxCore/Private/RivermaxOutputStream.cpp
+++ b/Engine/Plugins/VirtualProduction/Rivermax/RivermaxCore/Source/RivermaxCore/Private/RivermaxOutputStream.cpp
@@ -4,10 +4,11 @@
 
 #include "Async/Async.h"
 #include "CudaModule.h"
-#include "ID3D12DynamicRHI.h"
 #include "IRivermaxCoreModule.h"
+#include "IRivermaxBoundaryMonitor.h"
 #include "IRivermaxManager.h"
 #include "Misc/ByteSwap.h"
+#include "RivermaxFrameManager.h"
 #include "RivermaxLog.h"
 #include "RivermaxPTPUtils.h"
 #include "RivermaxTypes.h"
@@ -320,26 +321,19 @@ namespace UE::RivermaxCore::Private
 			return false;
 		}
 
-		// We do (try) to make gpu allocations here to let the capturer know if we require it or not.
-		if (RivermaxModule->GetRivermaxManager()->IsGPUDirectOutputSupported() && Options.bUseGPUDirect)
+		if (!SetupFrameManagement())
 		{
-			bUseGPUDirect = AllocateGPUBuffers();
+			return false;
 		}
 
+		// Enable frame boundary monitoring
+		MonitoringGuid = RivermaxModule->GetRivermaxBoundaryMonitor().StartMonitoring(Options.FrameRate);
+
 		Async(EAsyncExecution::TaskGraph, [this]()
 		{
-			//Create event to wait on when no frames are available to send
-			ReadyToSendEvent = FPlatformProcess::GetSynchEventFromPool();
-
 			TAnsiStringBuilder<2048> SDPDescription;
 			UE::RivermaxCore::Private::Utils::StreamOptionsToSDPDescription(Options, SDPDescription);
 
-			// Initialize buffers in system memory if we're not using gpudirect
-			if(bUseGPUDirect == false)
-			{
-				AllocateSystemBuffers();
-			}
-
 			// Initialize internal memory and rivermax configuration 
 			if(InitializeStreamMemoryConfig())
 			{
@@ -380,7 +374,7 @@ namespace UE::RivermaxCore::Private
 							, bUseGPUDirect ? TEXT(" using GPUDirect") : TEXT(""));
 
 						bIsActive = true;
-						RivermaxThread.Reset(FRunnableThread::Create(this, TEXT("Rmax OutputStream Thread"), 128 * 1024, TPri_AboveNormal, FPlatformAffinity::GetPoolThreadMask()));
+						RivermaxThread.Reset(FRunnableThread::Create(this, TEXT("Rmax OutputStream Thread"), 128 * 1024, TPri_TimeCritical, FPlatformAffinity::GetPoolThreadMask()));
 					}
 				}
 				else
@@ -401,17 +395,16 @@ namespace UE::RivermaxCore::Private
 		{
 			Stop();
 			
-			ReadyToSendEvent->Trigger();
+			FrameAvailableSignal->Trigger();
+			FrameReadyToSendSignal->Trigger();
 			RivermaxThread->Kill(true);
-			
 			RivermaxThread.Reset();
 
+			CleanupFrameManagement();
 
-			FPlatformProcess::ReturnSynchEventToPool(ReadyToSendEvent);
-			ReadyToSendEvent = nullptr;
+			RivermaxModule->GetRivermaxBoundaryMonitor().StopMonitoring(MonitoringGuid, Options.FrameRate);
+			
 			UE_LOG(LogRivermax, Log, TEXT("Rivermax Output stream has shutdown"));
-
-			DeallocateBuffers();
 		}
 	}
 
@@ -419,29 +412,13 @@ namespace UE::RivermaxCore::Private
 	{
 		TRACE_CPUPROFILER_EVENT_SCOPE(FRivermaxOutputStream::PushVideoFrame);
 
-		// Notify owner that new frame is about to be enqueued
-		Listener->OnPreFrameEnqueue();
-
-		if (TSharedPtr<FRivermaxOutputFrame> AvailableFrame = GetNextAvailableFrame(NewFrame.FrameIdentifier))
+		const bool bSuccess = FrameManager->SetFrameData(NewFrame);
+		if (!bSuccess)
 		{
-			const int32 Stride = GetStride();
-			FMemory::Memcpy(AvailableFrame->VideoBuffer, NewFrame.VideoBuffer, Options.AlignedResolution.Y * Stride);
-
-			AvailableFrame->bIsVideoBufferReady = true;
-
-			//If Frame ready to be sent
-			if(AvailableFrame->IsReadyToBeSent())
-			{
-				const FString TraceName = FString::Format(TEXT("FRivermaxOutputStream::PushFrame {0}"), { AvailableFrame->FrameIndex });
-				TRACE_CPUPROFILER_EVENT_SCOPE_TEXT(*TraceName);
-
-				MarkFrameToBeSent(MoveTemp(AvailableFrame));
-			}
-
-			return true;
+			Listener->OnStreamError();
 		}
-		
-		return false;
+
+		return bSuccess;
 	}
 
 	void FRivermaxOutputStream::Process_AnyThread()
@@ -485,6 +462,8 @@ namespace UE::RivermaxCore::Private
 
 			const FString TraceName = FString::Format(TEXT("RmaxOutput::SendFrame {0}"), { CurrentFrame->FrameIndex });
 			TRACE_CPUPROFILER_EVENT_SCOPE_TEXT(*TraceName);
+			TRACE_CPUPROFILER_EVENT_SCOPE_TEXT(*FString::Printf(TEXT("MediaCapturePipe: %u"), CurrentFrame->FrameIdentifier));
+
 			do 
 			{
 				if (bIsActive)
@@ -517,27 +496,6 @@ namespace UE::RivermaxCore::Private
 		}
 	}
 
-	void FRivermaxOutputStream::AllocateSystemBuffers()
-	{
-		const int32 Stride = GetStride();
-		const int32 FrameAllocSize = Options.AlignedResolution.Y * Stride;
-		AvailableFrames.Reserve(Options.NumberOfBuffers);
-
-		const auto SystemDeallocator = [](void* Buffer)
-		{
-			FMemory::Free(Buffer);
-		};
-
-		for (int32 Index = 0; Index < Options.NumberOfBuffers; ++Index)
-		{
-			TSharedPtr<FRivermaxOutputFrame> Frame = MakeShared<FRivermaxOutputFrame>(Index, SystemDeallocator);
-
-			constexpr uint32 CacheLineSize = PLATFORM_CACHE_LINE_SIZE;
-			Frame->VideoBuffer = FMemory::Malloc(FrameAllocSize, CacheLineSize);
-			AvailableFrames.Add(MoveTemp(Frame));
-		}
-	}
-
 	bool FRivermaxOutputStream::InitializeStreamMemoryConfig()
 	{
 		using namespace UE::RivermaxCore::Private::Utils;
@@ -653,7 +611,7 @@ namespace UE::RivermaxCore::Private
 			Block.chunks_num = StreamMemory.ChunksPerMemoryBlock;
 			Block.app_hdr_size_arr = StreamMemory.HeaderSizes.GetData();
 			Block.data_size_arr = StreamMemory.PayloadSizes.GetData();
-			Block.data_ptr = AvailableFrames[BlockIndex]->VideoBuffer;
+			Block.data_ptr = FrameManager->GetFrame(BlockIndex)->VideoBuffer;
 			
 			StreamMemory.RTPHeaders[BlockIndex].SetNumZeroed(StreamMemory.PacketsPerFrame);
 			Block.app_hdr_ptr = &StreamMemory.RTPHeaders[BlockIndex][0];
@@ -672,68 +630,6 @@ namespace UE::RivermaxCore::Private
 		NextFrame->HeaderPtr = nullptr;
 	}
 
-	TSharedPtr<FRivermaxOutputFrame> FRivermaxOutputStream::GetNextFrameToSend()
-	{
-		TRACE_CPUPROFILER_EVENT_SCOPE(FRivermaxOutputStream::GetNextFrameToSend);
-
-		TSharedPtr<FRivermaxOutputFrame> FrameToSend;
-
-		FScopeLock Lock(&FrameCriticalSection);
-		if (FramesToSend.IsEmpty() == false)
-		{
-			// Pick oldest frame of the lot. 
-			// Todo Sort array when modified to always pick first item.
-			uint32 OldestIdentifier = TNumericLimits<uint32>::Max();
-			TArray<TSharedPtr<FRivermaxOutputFrame>>::TIterator It = FramesToSend.CreateIterator();
-			for (It; It; ++It)
-			{
-				TSharedPtr<FRivermaxOutputFrame>& ItFrame = *It;
-				if (ItFrame->FrameIdentifier < OldestIdentifier)
-				{
-					FrameToSend = *It;
-					OldestIdentifier = FrameToSend->FrameIdentifier;
-				}
-			}
-
-			if (ensure(FrameToSend))
-			{
-				FramesToSend.Remove(FrameToSend);
-				InitializeNextFrame(FrameToSend);
-			}
-		}
-
-		return FrameToSend;
-	}
-
-	TSharedPtr<FRivermaxOutputFrame> FRivermaxOutputStream::GetNextAvailableFrame(uint32 InFrameIdentifier)
-	{
-		TRACE_CPUPROFILER_EVENT_SCOPE(FRivermaxOutputStream::GetNextAvailableFrame);
-		
-		TSharedPtr<FRivermaxOutputFrame> NextFrame;
-
-		{
-			FScopeLock Lock(&FrameCriticalSection);
-
-			//Find matching frame identifier
-			if(TSharedPtr<FRivermaxOutputFrame>* MatchingFrame = AvailableFrames.FindByPredicate([InFrameIdentifier](const TSharedPtr<FRivermaxOutputFrame>& Other){ return Other->FrameIdentifier == InFrameIdentifier; }))
-			{
-				NextFrame = *MatchingFrame;
-			}
-			else if(TSharedPtr<FRivermaxOutputFrame>* EmptyFrame = AvailableFrames.FindByPredicate([](const TSharedPtr<FRivermaxOutputFrame>& Other){ return Other->FrameIdentifier == FRivermaxOutputFrame::InvalidIdentifier; })) //Find an empty frame
-			{
-				NextFrame = *EmptyFrame;
-				NextFrame->FrameIdentifier = InFrameIdentifier;
-			}
-			else
-			{
-				// Could reuse a frame not yet ready to be sent
-				// Could also look into stealing a frame ready to be sent bot not sent yet
-			}
-		}
-
-		return NextFrame;
-	}
-
 	void FRivermaxOutputStream::BuildRTPHeader(FRawRTPHeader& OutHeader) const
 	{
 		using namespace UE::RivermaxCore::Private::Utils;
@@ -747,7 +643,7 @@ namespace UE::RivermaxCore::Private
 
 		// For now, in order to be able to use a framelocked input, we pipe frame number in the timestamp for a UE-UE interaction
 		// Follow up work to investigate adding this in RTP header
-		uint64 InputTime = StreamData.NextAlignmentPointNanosec;
+		uint64 InputTime = StreamData.NextScheduleTimeNanosec;
 		if (Options.bDoFrameCounterTimestamping)
 		{
 			InputTime = UE::RivermaxCore::GetAlignmentPointFromFrameNumber(CurrentFrame->FrameIdentifier, Options.FrameRate);
@@ -939,7 +835,7 @@ namespace UE::RivermaxCore::Private
 			{
 				UE_LOG(LogRivermax, Error, TEXT("Invalid error happened while trying to get next chunks. Status: %d"), Status);
 				Listener->OnStreamError();
-				bIsActive = false;
+				Stop();
 			}
 		} while (Status != RMAX_OK && bIsActive);
 	}
@@ -970,7 +866,7 @@ namespace UE::RivermaxCore::Private
 			const rmax_commit_flags_t CommitFlags{};
 			if (ScheduleTime != 0)
 			{
-				TRACE_BOOKMARK(TEXT("Sched A: %llu, C: %llu"), StreamData.NextAlignmentPointNanosec, CurrentTimeNanosec);
+				//TRACE_BOOKMARK(TEXT("Sched A: %llu, C: %llu"), StreamData.NextAlignmentPointNanosec, CurrentTimeNanosec);
 				if (ScheduleTime <= CurrentTimeNanosec)
 				{
 					ScheduleTime = 0;
@@ -994,13 +890,13 @@ namespace UE::RivermaxCore::Private
 			{
 				UE_LOG(LogRivermax, Error, TEXT("Completion issue while trying to commit next round of chunks."));
 				Listener->OnStreamError();
-				bIsActive = false;
+				Stop();
 			}
 			else
 			{
 				UE_LOG(LogRivermax, Error, TEXT("Unhandled error (%d) while trying to commit next round of chunks."), Status);
 				Listener->OnStreamError();
-				bIsActive = false;
+				Stop();
 			}
 
 		} while (Status != RMAX_OK && bIsActive);
@@ -1027,8 +923,8 @@ namespace UE::RivermaxCore::Private
 	{
 		// Initial wait for a frame to be produced
 		{
-			TRACE_CPUPROFILER_EVENT_SCOPE(FRivermaxOutputStream::InitialWait);
-			ReadyToSendEvent->Wait();
+			TRACE_CPUPROFILER_EVENT_SCOPE(Rmax::InitialWait);
+			FrameReadyToSendSignal->Wait();
 		}
 
 		while (bIsActive)
@@ -1054,105 +950,120 @@ namespace UE::RivermaxCore::Private
 
 	void FRivermaxOutputStream::PrepareNextFrame()
 	{
-		// We always want to send out a frame at the desired interval. If a new frame is not ready, reuse last one. 
-		bool bHasFrameToSend = false;
+		switch (Options.AlignmentMode)
 		{
-			FScopeLock Lock(&FrameCriticalSection);
-			bHasFrameToSend = FramesToSend.IsEmpty() == false;
+			case ERivermaxAlignmentMode::FrameCreation:
+			{
+				PrepareNextFrame_FrameCreation();
+				break;
+			}
+			case ERivermaxAlignmentMode::AlignmentPoint:
+			{
+				PrepareNextFrame_AlignmentPoint();
+				break;
+			}
+			default:
+			{
+				checkNoEntry();
+			}
 		}
+	}
 
-		const uint64 CurrentTimeNanosec = RivermaxModule->GetRivermaxManager()->GetTime();
-
-		// Different scenarios here.
+	void FRivermaxOutputStream::PrepareNextFrame_FrameCreation()
+	{
 		// When aligning on frame creation, we will always wait for a frame to be available.
-		// When aligning on alignment points, we either repeat the last one if none is available and continuous mode is on
-		// or skip sending one if continuous mode is off.
-		if (Options.AlignmentMode == ERivermaxAlignmentMode::FrameCreation)
+		TRACE_CPUPROFILER_EVENT_SCOPE(Rmax::WaitForReadyFrame);
+		TSharedPtr<FRivermaxOutputFrame> NextFrameToSend = FrameManager->GetReadyFrame();
+		while (!NextFrameToSend && bIsActive)
 		{
-			if (!bHasFrameToSend)
-			{
-				// No frame in this mode means we are going to wait for it
-				//First iteration before first frame available
-				TRACE_CPUPROFILER_EVENT_SCOPE(FRivermaxOutputStream::WaitForNewFrame);
-				do
-				{
-					FPlatformProcess::SleepNoStats(SleepTimeSeconds);
+			FrameReadyToSendSignal->Wait();
+			NextFrameToSend = FrameManager->GetReadyFrame();
+		}
 
-					// Early exit if stream isn't active anymore
-					if (!bIsActive)
-					{
-						return;
-					}
+		// In frame creation alignment, we always release the last frame sent
+		if (CurrentFrame.IsValid())
+		{
+			FrameManager->MarkAsSent(CurrentFrame);
+		}
 
-					{
-						FScopeLock Lock(&FrameCriticalSection);
-						bHasFrameToSend = FramesToSend.IsEmpty() == false;
-					}
+		// Make the next frame to send the current one and update its state
+		if (NextFrameToSend)
+		{
+			CurrentFrame = MoveTemp(NextFrameToSend);
+			FrameManager->MarkAsSending(CurrentFrame);
 
-				} while (bHasFrameToSend == false);
-			}
+			InitializeNextFrame(CurrentFrame);
 		}
+	}
 
+	void FRivermaxOutputStream::PrepareNextFrame_AlignmentPoint()
+	{
+		TRACE_CPUPROFILER_EVENT_SCOPE(Rmax::GetNextFrame_AlignmentPoint);
+
+		// When aligning on alignment points:
+		// We prepare to send the next frame that is ready if there is one available
+		// if none are available and bDoContinuousOutput == true
+		//		Repeat the last frame
+		// if none are available and bDoContinuousOutput == false
+		//		Don't send a frame and go back waiting for the next alignment point
+		
+		TSharedPtr<FRivermaxOutputFrame> NextFrameToSend = FrameManager->GetReadyFrame();
+
+		// If we have a new frame, release the previous one
+		// If we don't have a frame and we're not doing continuous output, we release it and we won't send a new one
+		// If we don't have a frame but we are doing continuous output, we will reschedule the current one, so no release.
+		if (!Options.bDoContinuousOutput || NextFrameToSend)
 		{
+			if (CurrentFrame)
 			{
-				FScopeLock Lock(&FrameCriticalSection);
-				bHasFrameToSend = FramesToSend.IsEmpty() == false;
+				FrameManager->MarkAsSent(CurrentFrame);
+				CurrentFrame.Reset();
 			}
 
-			// If we have a frame or we're not doing continuous output, we always release the last frame.
-			if (!Options.bDoContinuousOutput || bHasFrameToSend)
+			// Make the next frame to send the current one and update its state
+			if (NextFrameToSend)
 			{
-				// Not doing a continuous output and no frame to send, release the last one sent if any and move on waiting
-				if (CurrentFrame.IsValid())
-				{
-					//Release last frame sent. We keep hold to avoid overwriting it as rivermax is sending it
-					CurrentFrame->Reset();
-					{
-						FScopeLock Lock(&FrameCriticalSection);
-						AvailableFrames.Add(MoveTemp(CurrentFrame));
-						CurrentFrame.Reset();
-					}
-				}
-
-				CurrentFrame = GetNextFrameToSend();
-			}
-			else
-			{
-				// No frame to send, keep last one and restarts its internal counters
-				UE_LOG(LogRivermax, Verbose, TEXT("No frame to send. Reusing last frame '%d'"), CurrentFrame->FrameIndex);
+				CurrentFrame = MoveTemp(NextFrameToSend);
+				FrameManager->MarkAsSending(CurrentFrame);
 				InitializeNextFrame(CurrentFrame);
+			}
+		}
+		else
+		{
+			// We will resend the last one so just reinitialize it to resend
+			InitializeNextFrame(CurrentFrame);
+			
+			// No frame to send, keep last one and restart its internal counters
+			UE_LOG(LogRivermax, Verbose, TEXT("No frame to send. Reusing last frame '%d' with Id %u"), CurrentFrame->FrameIndex, CurrentFrame->FrameIdentifier);
 
-				// Since we want to resend last frame, we need to fast forward chunk pointer to re-point to the one we just sent
-				rmax_status_t Status;
-				bool bHasAddedTrace = false;
-				do
+			// Since we want to resend last frame, we need to fast forward chunk pointer to re-point to the one we just sent
+			rmax_status_t Status;
+			bool bHasAddedTrace = false;
+			do
+			{
+				Status = rmax_out_skip_chunks(StreamId, StreamMemory.ChunksPerFrameField * (Options.NumberOfBuffers - 1));
+				if (Status != RMAX_OK)
 				{
-					Status = rmax_out_skip_chunks(StreamId, StreamMemory.ChunksPerFrameField * (Options.NumberOfBuffers - 1));
-					if (Status != RMAX_OK)
+					if (Status == RMAX_ERR_NO_FREE_CHUNK)
 					{
-						if (Status == RMAX_ERR_NO_FREE_CHUNK)
-						{
-							// Wait until there are enough free chunk to be skipped
-							if (!bHasAddedTrace)
-							{
-								UE_LOG(LogRivermax, Warning, TEXT("No chunks ready to skip. Waiting"));
-								TRACE_CPUPROFILER_EVENT_SCOPE(PrepareNextFrame::NoFreeChunk);
-								bHasAddedTrace = true;
-							}
-						}
-						else
+						// Wait until there are enough free chunk to be skipped
+						if (!bHasAddedTrace)
 						{
-							ensure(false);
-							UE_LOG(LogRivermax, Error, TEXT("Invalid error happened while trying to skip chunks. Status: %d."), Status);
-							Listener->OnStreamError();
-							bIsActive = false;
+							UE_LOG(LogRivermax, Warning, TEXT("No chunks ready to skip. Waiting"));
+							TRACE_CPUPROFILER_EVENT_SCOPE(PrepareNextFrame::NoFreeChunk);
+							bHasAddedTrace = true;
 						}
 					}
-				} while (Status != RMAX_OK && bIsActive);
-			}
+					else
+					{
+						ensure(false);
+						UE_LOG(LogRivermax, Error, TEXT("Invalid error happened while trying to skip chunks. Status: %d."), Status);
+						Listener->OnStreamError();
+						Stop();
+					}
+				}
+			} while (Status != RMAX_OK && bIsActive);
 		}
-
-		//TRACE_BOOKMARK(TEXT("Schedulin %llu"), GetTimestampFromTime(StreamData.NextAlignmentPointNanosec, MediaClockSampleRate));
 	}
 
 	void FRivermaxOutputStream::InitializeStreamTimingSettings()
@@ -1249,265 +1160,47 @@ namespace UE::RivermaxCore::Private
 		}
 	}
 
-	bool FRivermaxOutputStream::PushGPUVideoFrame(const FRivermaxOutputVideoFrameInfo& NewFrame, FBufferRHIRef CapturedBuffer)
-	{
-		TRACE_CPUPROFILER_EVENT_SCOPE(FRivermaxOutputStream::PushGPUVideoFrame);
-
-		const int32 Stride = GetStride();
-		if (TSharedPtr<FRivermaxOutputFrame> AvailableFrame = GetNextAvailableFrame(NewFrame.FrameIdentifier))
-		{
-			const FString TraceName = FString::Format(TEXT("FRivermaxOutputStream::PushFrame {0}|{1}"), { AvailableFrame->FrameIndex, NewFrame.FrameIdentifier });
-			TRACE_CPUPROFILER_EVENT_SCOPE_TEXT(*TraceName);
-
-			FCUDAModule& CudaModule = FModuleManager::GetModuleChecked<FCUDAModule>("CUDA");
-			CUresult Result = CudaModule.DriverAPI()->cuCtxPushCurrent(CudaModule.GetCudaContext());
-			
-			void* MappedPointer = GetMappedAddress(CapturedBuffer);
-			if (MappedPointer == nullptr)
-			{
-				UE_LOG(LogRivermax, Error, TEXT("Failed to find a mapped memory address for captured buffer. Stopping capture."));
-				Listener->OnStreamError();
-				return false;
-			}
-
-			const CUdeviceptr CudaMemoryPointer = reinterpret_cast<CUdeviceptr>(MappedPointer);
-			Result = CudaModule.DriverAPI()->cuMemcpyDtoDAsync(reinterpret_cast<CUdeviceptr>(AvailableFrame->VideoBuffer), CudaMemoryPointer, Options.AlignedResolution.Y * Stride, reinterpret_cast<CUstream>(GPUStream));
-			if (Result != CUDA_SUCCESS)
-			{
-				UE_LOG(LogRivermax, Error, TEXT("Failed to copy captured bufer to cuda memory. Stopping capture. Error: %d"), Result);
-				Listener->OnStreamError();
-				return false;
-			}
-
-
-			// Callback called by Cuda when stream work has completed on cuda engine (MemCpy -> Callback)
-			// Once Memcpy has been done, we know we can mark that memory as available to be sent. 
-			
-			auto CudaCallback = [](void* userData)
-			{
-				TRACE_CPUPROFILER_EVENT_SCOPE(CudaWorkDoneCallback);
-
-				FRivermaxOutputStream* Stream = reinterpret_cast<FRivermaxOutputStream*>(userData);
-
-				// Notify owner that new frame is about to be enqueued
-				Stream->Listener->OnPreFrameEnqueue();
-
-				TOptional<uint32> TargetFrameIdentifier = Stream->PendingIdentifiers.Dequeue();
-				if (TargetFrameIdentifier.IsSet())
-				{
-					if (TSharedPtr<FRivermaxOutputFrame> AvailableFrame = Stream->GetNextAvailableFrame(TargetFrameIdentifier.GetValue()))
-					{
-						AvailableFrame->bIsVideoBufferReady = true;
-						if (AvailableFrame->IsReadyToBeSent())
-						{
-							const FString TraceName = FString::Format(TEXT("FRivermaxOutputStream::PushFrame {0}"), { AvailableFrame->FrameIndex });
-							TRACE_CPUPROFILER_EVENT_SCOPE_TEXT(*TraceName);
-
-							Stream->MarkFrameToBeSent(MoveTemp(AvailableFrame));
-						}
-					}
-				}
-			};
-			
-			// Add pending frame for cuda callback 
-			PendingIdentifiers.Enqueue(AvailableFrame->FrameIdentifier);
-
-			// Schedule a callback to make the frame available
-			CudaModule.DriverAPI()->cuLaunchHostFunc(reinterpret_cast<CUstream>(GPUStream), CudaCallback, this);
-			
-			FCUDAModule::CUDA().cuCtxPopCurrent(nullptr);
-
-			return true;
-		}
-
-		return false;
-	}
-
 	bool FRivermaxOutputStream::IsGPUDirectSupported() const
 	{
 		return bUseGPUDirect;
 	}
 
-	bool FRivermaxOutputStream::AllocateGPUBuffers()
+	bool FRivermaxOutputStream::SetupFrameManagement()
 	{
-		// Allocate a single memory space that will contain all frame buffers
-		
-		TRACE_CPUPROFILER_EVENT_SCOPE(FRivermaxOutputStream::AllocateGPUBuffers);
-
-		const ERHIInterfaceType RHIType = RHIGetInterfaceType();
-		if (RHIType != ERHIInterfaceType::D3D12)
-		{
-			UE_LOG(LogRivermax, Warning, TEXT("Can't initialize output to use GPUDirect. RHI is %d but only Dx12 is supported at the moment."), int(RHIType));
-			return false;
-		}
-
-		const int32 Stride = GetStride();
-		const int32 FrameAllocSize = Options.AlignedResolution.Y * Stride;
+		TRACE_CPUPROFILER_EVENT_SCOPE(FRivermaxOutputStream::SetupFrameManagement);
 
-		FCUDAModule& CudaModule = FModuleManager::GetModuleChecked<FCUDAModule>("CUDA");
-		CudaModule.DriverAPI()->cuCtxPushCurrent(CudaModule.GetCudaContext());
+		FrameManager = MakeUnique<FFrameManager>();
 
-		// Todo: Add support for mgpu. For now, this will not work unless the memcpy does implicitely a cross gpu transfer.
-		const int GPUIndex = CudaModule.GetCudaDeviceIndex();
-		CUdevice CudaDevice;
-		CUresult Status = CudaModule.DriverAPI()->cuDeviceGet(&CudaDevice, GPUIndex);
-		if(Status != CUDA_SUCCESS)
-		{
-			UE_LOG(LogRivermax, Warning, TEXT("Can't initialize output to use GPUDirect. Failed to get a Cuda device for GPU %d. Status: %d"), GPUIndex, Status);
-			return false;
-		}
-
-		CUmemAllocationProp AllocationProperties = {};
-		AllocationProperties.type = CU_MEM_ALLOCATION_TYPE_PINNED;
-		AllocationProperties.allocFlags.gpuDirectRDMACapable = 1;
-		AllocationProperties.allocFlags.usage = 0;
-		AllocationProperties.location.type = CU_MEM_LOCATION_TYPE_DEVICE;
-		AllocationProperties.location.id = CudaDevice;
-
-		// Get memory granularity required for cuda device. We need to align allocation with this.
-		size_t Granularity;
-		Status = CudaModule.DriverAPI()->cuMemGetAllocationGranularity(&Granularity, &AllocationProperties, CU_MEM_ALLOC_GRANULARITY_RECOMMENDED);
-		if(Status != CUDA_SUCCESS)
-		{
-			UE_LOG(LogRivermax, Warning, TEXT("Can't initialize output to use GPUDirect. Failed to get allocation granularity. Status: %d"), Status);
-			return false;
-		}
-
-		// Cuda requires allocated memory to be aligned with a certain granularity
-		// We align each frame size to the desired granularity and multiply that by number of buffer
-		// This causes more memory to be allocated but doing a single allocation fails rmax stream creation
-		const size_t CudaAlignedFrameSize = (FrameAllocSize % Granularity) ? FrameAllocSize + (Granularity - (FrameAllocSize % Granularity)) : FrameAllocSize;
-		const size_t TotalCudaAllocSize = CudaAlignedFrameSize * Options.NumberOfBuffers;
-
-		// Reserve contiguous memory to contain required number of buffers. 
-		CUdeviceptr CudaBaseAddress;
-		constexpr CUdeviceptr InitialAddress = 0;
-		constexpr int32 Flags = 0;
-		Status = CudaModule.DriverAPI()->cuMemAddressReserve(&CudaBaseAddress, TotalCudaAllocSize, Granularity, InitialAddress, Flags);
-		if (Status != CUDA_SUCCESS)
-		{
-			UE_LOG(LogRivermax, Warning, TEXT("Can't initialize output to use GPUDirect. Failed to reserve memory for %d bytes. Status: %d"), TotalCudaAllocSize, Status);
-			return false;
-		}
-
-		// Make the allocation on device memory
-		CUmemGenericAllocationHandle Handle;
-		Status = CudaModule.DriverAPI()->cuMemCreate(&Handle, TotalCudaAllocSize, &AllocationProperties, Flags);
-		if (Status != CUDA_SUCCESS)
-		{
-			UE_LOG(LogRivermax, Warning, TEXT("Can't initialize output to use GPUDirect. Failed to create memory on device. Status: %d"), Status);
-			return false;
-		}
-		
-		bool bExit = false;
-		constexpr int32 Offset = 0;
-		Status = CudaModule.DriverAPI()->cuMemMap(CudaBaseAddress, TotalCudaAllocSize, Offset, Handle, Flags);
-		if (Status != CUDA_SUCCESS)
-		{
-			UE_LOG(LogRivermax, Warning, TEXT("Can't initialize output to use GPUDirect. Failed to map memory. Status: %d"), Status);
-			// Need to release handle no matter what
-			bExit = true;
-		}
-		
-		// Cache to know we need to unmap/deallocate even if it fails down the road
-		CudaAllocatedMemory = TotalCudaAllocSize;
-		CudaAllocatedMemoryBaseAddress = reinterpret_cast<void*>(CudaBaseAddress);
-
-		Status = CudaModule.DriverAPI()->cuMemRelease(Handle);
-		if (Status != CUDA_SUCCESS)
-		{
-			UE_LOG(LogRivermax, Warning, TEXT("Can't initialize output to use GPUDirect. Failed to release handle. Status: %d"), Status);
-			return false;
-		}
-
-		if(bExit)
-		{
-			return false;
-		}
-
-		// Setup access description.
-		CUmemAccessDesc MemoryAccessDescription = {};
-		MemoryAccessDescription.flags = CU_MEM_ACCESS_FLAGS_PROT_READWRITE;
-		MemoryAccessDescription.location.type = CU_MEM_LOCATION_TYPE_DEVICE;
-		MemoryAccessDescription.location.id = CudaDevice;
-		constexpr int32 Count = 1;
-		Status = CudaModule.DriverAPI()->cuMemSetAccess(CudaBaseAddress, TotalCudaAllocSize, &MemoryAccessDescription, Count);
-		if (Status != CUDA_SUCCESS)
-		{
-			UE_LOG(LogRivermax, Warning, TEXT("Can't initialize output to use GPUDirect. Failed to configure memory access. Status: %d"), Status);
-			return false;
-		}
-
-		CUstream CudaStream;
-		Status = CudaModule.DriverAPI()->cuStreamCreate(&CudaStream, CU_STREAM_NON_BLOCKING);
-		if (Status != CUDA_SUCCESS)
-		{
-			UE_LOG(LogRivermax, Warning, TEXT("Can't initialize output to use GPUDirect. Failed to create its stream. Status: %d"), Status);
-			return false;
-		}
-
-		GPUStream = CudaStream;
-		
-		Status = CudaModule.DriverAPI()->cuCtxSynchronize();
-		if (Status != CUDA_SUCCESS)
+		// We do (try) to make gpu allocations here to let the capturer know if we require it or not.
+		bool bTryGPUDirect = RivermaxModule->GetRivermaxManager()->IsGPUDirectOutputSupported() && Options.bUseGPUDirect;
+		if (bTryGPUDirect)
 		{
-			UE_LOG(LogRivermax, Warning, TEXT("Can't initialize output to use GPUDirect. Failed to synchronize context. Status: %d"), Status);
-			return false;
+			const ERHIInterfaceType RHIType = RHIGetInterfaceType();
+			if (RHIType != ERHIInterfaceType::D3D12)
+			{
+				UE_LOG(LogRivermax, Warning, TEXT("Can't initialize output to use GPUDirect. RHI is %d but only Dx12 is supported at the moment."), RHIType);
+				bTryGPUDirect = false;
+			}
 		}
 
-		// Now that Cuda memory was allocated, we can allocate our buffers with an address in cuda space
-		AvailableFrames.Reserve(Options.NumberOfBuffers);
-		for (int32 Index = 0; Index < Options.NumberOfBuffers; ++Index)
-		{
-			TFunction<void(void*)> DeallocatorFunc = nullptr; // We handle a single deallocation for gpu memory
-			TSharedPtr<FRivermaxOutputFrame> Frame = MakeShared<FRivermaxOutputFrame>(Index, DeallocatorFunc);
-			Frame->VideoBuffer = reinterpret_cast<void*>(CudaBaseAddress + (Index * CudaAlignedFrameSize));
-			AvailableFrames.Add(MoveTemp(Frame));
-		}		
+		FFrameManagerSetupArgs FrameManagerArgs;
+		FrameManagerArgs.Resolution = Options.AlignedResolution;
+		FrameManagerArgs.bTryGPUAllocation = bTryGPUDirect;
+		FrameManagerArgs.NumberOfFrames = Options.NumberOfBuffers;
+		FrameManagerArgs.Stride = GetStride();
+		FrameManagerArgs.OnFreeFrameDelegate = FOnFrameReadyDelegate::CreateRaw(this, &FRivermaxOutputStream::OnFrameReadyToBeUsed);
+		FrameManagerArgs.OnPreFrameReadyDelegate = FOnFrameReadyDelegate::CreateRaw(this, &FRivermaxOutputStream::OnPreFrameReadyToBeSent);
+		FrameManagerArgs.OnFrameReadyDelegate = FOnFrameReadyDelegate::CreateRaw(this, &FRivermaxOutputStream::OnFrameReadyToBeSent);
+		const EFrameMemoryLocation FrameLocation = FrameManager->Initialize(FrameManagerArgs);
+		bUseGPUDirect = FrameLocation == EFrameMemoryLocation::GPU;
 
-		CudaModule.DriverAPI()->cuCtxPopCurrent(nullptr);
-
-		return true;
+		return FrameLocation != EFrameMemoryLocation::None;
 	}
 
-	void FRivermaxOutputStream::DeallocateBuffers()
+	void FRivermaxOutputStream::CleanupFrameManagement()
 	{
-		if (CudaAllocatedMemory > 0)
-		{
-			FCUDAModule& CudaModule = FModuleManager::GetModuleChecked<FCUDAModule>("CUDA");
-			CudaModule.DriverAPI()->cuCtxPushCurrent(CudaModule.GetCudaContext());
-
-			const CUdeviceptr BaseAddress = reinterpret_cast<CUdeviceptr>(CudaAllocatedMemoryBaseAddress);
-			CUresult Status = CudaModule.DriverAPI()->cuMemUnmap(BaseAddress, CudaAllocatedMemory);
-			if (Status != CUDA_SUCCESS)
-			{
-				UE_LOG(LogRivermax, Warning, TEXT("Failed to unmap cuda memory. Status: %d"), Status);
-			}
-
-			Status = CudaModule.DriverAPI()->cuMemAddressFree(BaseAddress, CudaAllocatedMemory);
-			if (Status != CUDA_SUCCESS)
-			{
-				UE_LOG(LogRivermax, Warning, TEXT("Failed to free cuda memory. Status: %d"), Status);
-			}
-
-			for (const TPair<FBufferRHIRef, void*>& Entry : BufferCudaMemoryMap)
-			{
-				if(Entry.Value)
-				{
-					CudaModule.DriverAPI()->cuMemFree(reinterpret_cast<CUdeviceptr>(Entry.Value));
-				}
-			}
-			BufferCudaMemoryMap.Empty();
-
-			Status = CudaModule.DriverAPI()->cuStreamDestroy(reinterpret_cast<CUstream>(GPUStream));
-			if (Status != CUDA_SUCCESS)
-			{
-				UE_LOG(LogRivermax, Warning, TEXT("Failed to destroy cuda stream. Status: %d"), Status);
-			}
-			GPUStream = nullptr;
-
-			CudaModule.DriverAPI()->cuCtxPopCurrent(nullptr);
-		}
+		FrameManager->Cleanup();
+		FrameManager.Reset();
 	}
 
 	int32 FRivermaxOutputStream::GetStride() const
@@ -1516,87 +1209,6 @@ namespace UE::RivermaxCore::Private
 		return (Options.AlignedResolution.X / FormatInfo.PixelGroupCoverage) * FormatInfo.PixelGroupSize;
 	}
 
-	void* FRivermaxOutputStream::GetMappedAddress(const FBufferRHIRef& InBuffer)
-	{
-		// If we are here, d3d12 had to have been validated
-		const ERHIInterfaceType RHIType = RHIGetInterfaceType();
-		check(RHIType == ERHIInterfaceType::D3D12);
-
-		//Do we already have a mapped address for this buffer
-		if (BufferCudaMemoryMap.Find((InBuffer)) == nullptr)
-		{
-			int64 BufferMemorySize = 0;
-			CUexternalMemory MappedExternalMemory = nullptr;
-			HANDLE D3D12BufferHandle = 0;
-			CUDA_EXTERNAL_MEMORY_HANDLE_DESC CudaExtMemHandleDesc = {};
-
-			// Create shared handle for our buffer
-			{
-				TRACE_CPUPROFILER_EVENT_SCOPE(Rmax_D3D12CreateSharedHandle);
-
-				ID3D12Resource* NativeD3D12Resource = GetID3D12DynamicRHI()->RHIGetResource(InBuffer);
-				BufferMemorySize = GetID3D12DynamicRHI()->RHIGetResourceMemorySize(InBuffer);
-
-				TRefCountPtr<ID3D12Device> OwnerDevice;
-				HRESULT QueryResult;
-				if ((QueryResult = NativeD3D12Resource->GetDevice(IID_PPV_ARGS(OwnerDevice.GetInitReference()))) != S_OK)
-				{
-					UE_LOG(LogRivermax, Error, TEXT("Failed to get D3D12 device for captured buffer ressource: %d)"), QueryResult);
-					return nullptr;
-				}
-
-				if ((QueryResult = OwnerDevice->CreateSharedHandle(NativeD3D12Resource, NULL, GENERIC_ALL, NULL, &D3D12BufferHandle)) != S_OK)
-				{
-					UE_LOG(LogRivermax, Error, TEXT("Failed to create shared handle for captured buffer ressource: %d"), QueryResult);
-					return nullptr;
-				}
-
-				CudaExtMemHandleDesc.type = CU_EXTERNAL_MEMORY_HANDLE_TYPE_D3D12_RESOURCE;
-				CudaExtMemHandleDesc.handle.win32.name = nullptr;
-				CudaExtMemHandleDesc.handle.win32.handle = D3D12BufferHandle;
-				CudaExtMemHandleDesc.size = BufferMemorySize;
-				CudaExtMemHandleDesc.flags |= CUDA_EXTERNAL_MEMORY_DEDICATED;
-			}
-
-			{
-				TRACE_CPUPROFILER_EVENT_SCOPE(Rmax_CudaImportMemory);
-
-				const CUresult Result = FCUDAModule::CUDA().cuImportExternalMemory(&MappedExternalMemory, &CudaExtMemHandleDesc);
-				
-				if (D3D12BufferHandle)
-				{
-					CloseHandle(D3D12BufferHandle);
-				}
-
-				if (Result != CUDA_SUCCESS)
-				{
-					UE_LOG(LogRivermax, Error, TEXT("Failed to import shared buffer. Error: %d"), Result);
-					return nullptr;
-				}
-			}
-
-			{
-				TRACE_CPUPROFILER_EVENT_SCOPE(Rmax_MapCudaMemory);
-
-				CUDA_EXTERNAL_MEMORY_BUFFER_DESC BufferDescription = {};
-				BufferDescription.offset = 0;
-				BufferDescription.size = BufferMemorySize;
-				CUdeviceptr NewMemory;
-				const CUresult Result = FCUDAModule::CUDA().cuExternalMemoryGetMappedBuffer(&NewMemory, MappedExternalMemory, &BufferDescription);
-				if(Result != CUDA_SUCCESS || NewMemory == 0)
-				{
-					UE_LOG(LogRivermax, Error, TEXT("Failed to get shared buffer mapped memory. Error: %d"), Result);
-					return nullptr;
-				}
-				
-				BufferCudaMemoryMap.Add(InBuffer, reinterpret_cast<void*>(NewMemory));
-			}
-		}
-
-		// At this point, we have the mapped buffer in cuda space and we can use it to schedule a memcpy on cuda engine.
-		return BufferCudaMemoryMap[InBuffer];
-	}
-
 	void FRivermaxOutputStream::CalculateNextScheduleTime_AlignementPoints(uint64 CurrentClockTimeNanosec, uint64 CurrentFrameNumber)
 	{
 		// Frame number we will want to align with
@@ -1641,7 +1253,6 @@ namespace UE::RivermaxCore::Private
 					{
 						NextFrameNumber = CurrentFrameNumber + 1;
 					}
-
 				}
 				else
 				{
@@ -1682,15 +1293,34 @@ namespace UE::RivermaxCore::Private
 		}
 	}
 
-	void FRivermaxOutputStream::MarkFrameToBeSent(TSharedPtr<FRivermaxOutputFrame> ReadyFrame)
+	bool FRivermaxOutputStream::ReserveFrame(uint32 FrameIdentifier) const
+	{
+		TSharedPtr<FRivermaxOutputFrame> ReservedFrame = FrameManager->GetFreeFrame();
+		if (!ReservedFrame && Options.FrameLockingMode == EFrameLockingMode::BlockOnReservation)
+		{
+			TRACE_CPUPROFILER_EVENT_SCOPE(Rmax::WaitForAvailableFrame);
+			while(!ReservedFrame && bIsActive)
+			{
+				FrameAvailableSignal->Wait();
+				ReservedFrame = FrameManager->GetFreeFrame();
+			}
+		}
+		return ReservedFrame != nullptr;
+	}
+
+	void FRivermaxOutputStream::OnFrameReadyToBeSent()
+	{	
+		FrameReadyToSendSignal->Trigger();
+	}
+
+	void FRivermaxOutputStream::OnFrameReadyToBeUsed()
 	{
-		// Make frame available to be sent
-		FScopeLock Lock(&FrameCriticalSection);
-		ReadyFrame->ReadyTimestamp = RivermaxModule->GetRivermaxManager()->GetTime();
-		AvailableFrames.Remove(ReadyFrame);
-		FramesToSend.Add(MoveTemp(ReadyFrame));
-		ReadyToSendEvent->Trigger();
+		FrameAvailableSignal->Trigger();
 	}
 
+	void FRivermaxOutputStream::OnPreFrameReadyToBeSent()
+	{
+		Listener->OnPreFrameEnqueue();
+	}
 }
 
diff --git a/Engine/Plugins/VirtualProduction/Rivermax/RivermaxCore/Source/RivermaxCore/Private/RivermaxOutputStream.h b/Engine/Plugins/VirtualProduction/Rivermax/RivermaxCore/Source/RivermaxCore/Private/RivermaxOutputStream.h
--- a/Engine/Plugins/VirtualProduction/Rivermax/RivermaxCore/Source/RivermaxCore/Private/RivermaxOutputStream.h
+++ b/Engine/Plugins/VirtualProduction/Rivermax/RivermaxCore/Source/RivermaxCore/Private/RivermaxOutputStream.h
@@ -17,6 +17,8 @@ class IRivermaxCoreModule;
 
 namespace UE::RivermaxCore::Private
 {
+	class FFrameManager;
+
 	using UE::RivermaxCore::FRivermaxOutputStreamOptions;
 	using UE::RivermaxCore::FRivermaxOutputVideoFrameInfo;
 
@@ -97,8 +99,8 @@ namespace UE::RivermaxCore::Private
 		virtual bool Initialize(const FRivermaxOutputStreamOptions& Options, IRivermaxOutputStreamListener& InListener) override;
 		virtual void Uninitialize() override;
 		virtual bool PushVideoFrame(const FRivermaxOutputVideoFrameInfo& NewFrame) override;
-		virtual bool PushGPUVideoFrame(const FRivermaxOutputVideoFrameInfo& NewFrame, FBufferRHIRef CapturedBuffer) override;
 		virtual bool IsGPUDirectSupported() const override;
+		virtual bool ReserveFrame(uint32 FrameIdentifier) const override;
 		//~ End IRivermaxOutputStream interface
 
 		void Process_AnyThread();
@@ -118,24 +120,15 @@ namespace UE::RivermaxCore::Private
 		/** Initializes timing setup for this stream. TRO, frame interval etc... */
 		void InitializeStreamTimingSettings();
 
-		/** Allocates buffers on gpu for gpudirect usage */
-		bool AllocateGPUBuffers();
-
-		/** Allocates buffers on system memory */
-		void AllocateSystemBuffers();
+		/** Sets up frame management taking care of allocation, special cuda handling, etc... */
+		bool SetupFrameManagement();
 
-		/** Clean up allocated buffers */
-		void DeallocateBuffers();
+		/** Clean up frames */
+		void CleanupFrameManagement();
 
 		/** Resets NextFrame to be ready to send it out */
 		void InitializeNextFrame(const TSharedPtr<FRivermaxOutputFrame>& NextFrame);
 
-		/** Returns next frame ready to be sent */
-		TSharedPtr<FRivermaxOutputFrame> GetNextFrameToSend();
-
-		/** Returns next frame ready to be filled / written by the capture */
-		TSharedPtr<FRivermaxOutputFrame> GetNextAvailableFrame(uint32 InFrameIdentifier);
-
 		/** Fills RTP and SRD header using current state */
 		void BuildRTPHeader(FRawRTPHeader& OutHeader) const;
 
@@ -163,6 +156,12 @@ namespace UE::RivermaxCore::Private
 		/** Fetches next frame to send and prepares it for sending */
 		void PrepareNextFrame();
 
+		/** Returns next frame to send for frame creation alignment */
+		void PrepareNextFrame_FrameCreation();
+
+		/** Returns next frame to send for alignement point method. Can return nullptr */
+		void PrepareNextFrame_AlignmentPoint();
+
 		/** If enabled, print stats related to this stream */
 		void ShowStats();
 
@@ -172,11 +171,14 @@ namespace UE::RivermaxCore::Private
 		/** Get row stride for the current stream configuration */
 		int32 GetStride() const;
 
-		/** Get mapped address in cuda space for a given buffer. Cache will be updated if not found */
-		void* GetMappedAddress(const FBufferRHIRef& InBuffer);
+		/** Used to notify the listener that a frame is ready to be enqueued for transmission */
+		void OnPreFrameReadyToBeSent();
+		
+		/** Used to detect when a frame is now ready to be sent */
+		void OnFrameReadyToBeSent();
 
-		/** Makes a frame available to be sent, i.e. moved to the right container and mark its arrival time */
-		void MarkFrameToBeSent(TSharedPtr<FRivermaxOutputFrame> ReadyFrame);
+		/** Used to know when a frame is ready to be used and receive new data */
+		void OnFrameReadyToBeUsed();
 
 	private:
 
@@ -196,25 +198,25 @@ namespace UE::RivermaxCore::Private
 		rmax_stream_id StreamId;
 
 		/** Critical section to protect frames access */
-		FCriticalSection FrameCriticalSection;
+		mutable FCriticalSection FrameCriticalSection;
 
 		/** Current frame being sent */
 		TSharedPtr<FRivermaxOutputFrame> CurrentFrame;
 
-		/** Available frames to write memory to (Ready to be written) */
-		TArray<TSharedPtr<FRivermaxOutputFrame>> AvailableFrames;
-
-		/** Frames ready to be sent to rivermax (Ready to be read) */
-		TArray<TSharedPtr<FRivermaxOutputFrame>> FramesToSend;
-
 		/** Thread scheduling frame output */
 		TUniquePtr<FRunnableThread> RivermaxThread;
 
+		/** Manages allocation and memory manipulation of video frames */
+		TUniquePtr<FFrameManager> FrameManager;
+
 		/** Whether stream is active or not */
 		std::atomic<bool> bIsActive;
 
 		/** Event used to let scheduler that a frame is ready to be sent */
-		FEvent* ReadyToSendEvent = nullptr;
+		FEventRef FrameReadyToSendSignal = FEventRef(EEventMode::AutoReset);
+
+		/** Event used to unblock frame reservation as soon as one is free */
+		FEventRef FrameAvailableSignal = FEventRef(EEventMode::AutoReset);
 
 		/** Listener for this stream events */
 		IRivermaxOutputStreamListener* Listener = nullptr;
@@ -234,24 +236,12 @@ namespace UE::RivermaxCore::Private
 		/** Whether stream is using gpudirect to host memory consumed by Rivermax */
 		bool bUseGPUDirect = false;
 
-		/** Allocated memory base address used when it's time to free */
-		void* CudaAllocatedMemoryBaseAddress = nullptr;
-
-		/** Total allocated gpu memory. */
-		int32 CudaAllocatedMemory = 0;
-
-		/** Map between buffer we are sending and their mapped address in gpu space */
-		TMap<FBufferRHIRef, void*> BufferCudaMemoryMap;
-
-		/** Queued identifiers to be consumed by cuda callback when work has been completed */
-		TSpscQueue<uint32> PendingIdentifiers;
-
-		/** Cuda stream used for our operations */
-		void* GPUStream = nullptr;
-
 		/** Our own module pointer kept for ease of use */
 		IRivermaxCoreModule* RivermaxModule = nullptr;
 
+		/** Guid given by boundary monitoring handler to unregister ourselves */
+		FGuid MonitoringGuid;
+
 		/** Time to sleep when waiting for an operation to complete */
 		static constexpr double SleepTimeSeconds = 50.0 * 1E-6;
 	};
diff --git a/Engine/Plugins/VirtualProduction/Rivermax/RivermaxCore/Source/RivermaxCore/Public/IRivermaxBoundaryMonitor.h b/Engine/Plugins/VirtualProduction/Rivermax/RivermaxCore/Source/RivermaxCore/Public/IRivermaxBoundaryMonitor.h
new file mode 100644
--- /dev/null
+++ b/Engine/Plugins/VirtualProduction/Rivermax/RivermaxCore/Source/RivermaxCore/Public/IRivermaxBoundaryMonitor.h
@@ -0,0 +1,32 @@
+// Copyright Epic Games, Inc. All Rights Reserved.
+
+#pragma once
+
+#include "CoreTypes.h"
+
+struct FGuid;
+struct FFrameRate;
+
+namespace UE::RivermaxCore
+{
+	/** 
+	 * Frame boundary monitor interface used to
+	 * add trace bookmarks for monitored frame rates.
+	 * Boundaries are calculated using ST-2059 formula.
+	 */
+	class RIVERMAXCORE_API IRivermaxBoundaryMonitor
+	{
+	public:
+		virtual ~IRivermaxBoundaryMonitor() = default;
+
+		/** Global enable / disable of monitoring thread */
+		virtual void EnableMonitoring(bool bEnable) = 0;
+
+		/** Asks to start monitoring a particular frame rate */
+		virtual FGuid StartMonitoring(const FFrameRate& FrameRate) = 0;
+	
+		/** Removes a listener for a particular frame rate */
+		virtual void StopMonitoring(const FGuid& Requester, const FFrameRate& FrameRate) = 0;
+	};
+}
+
diff --git a/Engine/Plugins/VirtualProduction/Rivermax/RivermaxCore/Source/RivermaxCore/Public/IRivermaxCoreModule.h b/Engine/Plugins/VirtualProduction/Rivermax/RivermaxCore/Source/RivermaxCore/Public/IRivermaxCoreModule.h
--- a/Engine/Plugins/VirtualProduction/Rivermax/RivermaxCore/Source/RivermaxCore/Public/IRivermaxCoreModule.h
+++ b/Engine/Plugins/VirtualProduction/Rivermax/RivermaxCore/Source/RivermaxCore/Public/IRivermaxCoreModule.h
@@ -12,6 +12,7 @@
 namespace UE::RivermaxCore
 {
 	class IRivermaxManager;
+	class IRivermaxBoundaryMonitor;
 }
 
 
@@ -40,5 +41,8 @@ public:
 
 	/** Returns Rivermax manager singleton to query for stats, status, etc... */
 	virtual TSharedPtr<UE::RivermaxCore::IRivermaxManager> GetRivermaxManager() = 0;
+
+	/** Returns frame boundary monitor */
+	virtual UE::RivermaxCore::IRivermaxBoundaryMonitor& GetRivermaxBoundaryMonitor() = 0;
 };
 
diff --git a/Engine/Plugins/VirtualProduction/Rivermax/RivermaxCore/Source/RivermaxCore/Public/IRivermaxOutputStream.h b/Engine/Plugins/VirtualProduction/Rivermax/RivermaxCore/Source/RivermaxCore/Public/IRivermaxOutputStream.h
--- a/Engine/Plugins/VirtualProduction/Rivermax/RivermaxCore/Source/RivermaxCore/Public/IRivermaxOutputStream.h
+++ b/Engine/Plugins/VirtualProduction/Rivermax/RivermaxCore/Source/RivermaxCore/Public/IRivermaxOutputStream.h
@@ -18,6 +18,7 @@ namespace UE::RivermaxCore
 		uint32 Width = 0;
 		uint32 Stride = 0;
 		void* VideoBuffer = nullptr;
+		FBufferRHIRef GPUBuffer;
 	};
 
 	class RIVERMAXCORE_API IRivermaxOutputStreamListener
@@ -54,10 +55,21 @@ namespace UE::RivermaxCore
 		virtual bool PushVideoFrame(const FRivermaxOutputVideoFrameInfo& NewFrame) = 0;
 		
 		/** Pushes a frame to be captured from GPU memory */
-		virtual bool PushGPUVideoFrame(const FRivermaxOutputVideoFrameInfo& NewFrame, FBufferRHIRef CapturedBuffer) = 0;
+		UE_DEPRECATED("5.3", "Use PushVideoFrame instead and fill GPUBuffer with Capturedbuffer")
+		virtual bool PushGPUVideoFrame(const FRivermaxOutputVideoFrameInfo& NewFrame, FBufferRHIRef CapturedBuffer) final
+		{
+			return false;
+		}
 
 		/** Returns true if GPUDirect is supported */
 		virtual bool IsGPUDirectSupported() const = 0;
+
+		/** 
+		 * Tries to reserve a frame for the next capture 
+		 * If FRivermaxOutputStreamOptions::EFrameLockingMode equals BlockOnReservation, 
+		 * this method will block until a free frame is found. 
+		 */
+		virtual bool ReserveFrame(uint32 FrameIdentifier) const = 0;
 	};
 }
 
diff --git a/Engine/Plugins/VirtualProduction/Rivermax/RivermaxCore/Source/RivermaxCore/Public/RivermaxTypes.h b/Engine/Plugins/VirtualProduction/Rivermax/RivermaxCore/Source/RivermaxCore/Public/RivermaxTypes.h
--- a/Engine/Plugins/VirtualProduction/Rivermax/RivermaxCore/Source/RivermaxCore/Public/RivermaxTypes.h
+++ b/Engine/Plugins/VirtualProduction/Rivermax/RivermaxCore/Source/RivermaxCore/Public/RivermaxTypes.h
@@ -11,10 +11,22 @@ namespace UE::RivermaxCore
 {
 	enum class ERivermaxAlignmentMode
 	{
+		/** Aligns scheduling with ST2059 frame boundary formula */
 		AlignmentPoint,
+
+		/** Aligns scheduling with frame creation */
 		FrameCreation,
 	};
 
+	enum class EFrameLockingMode : uint8
+	{
+		/** If no frame available, continue */
+		FreeRun,
+
+		/** Blocks when reserving a frame slot. */
+		BlockOnReservation,
+	};
+
 	struct RIVERMAXCORE_API FRivermaxInputStreamOptions
 	{
 		/** Stream FrameRate */
@@ -77,6 +89,9 @@ namespace UE::RivermaxCore
 		/** Method used to align output stream */
 		ERivermaxAlignmentMode AlignmentMode = ERivermaxAlignmentMode::AlignmentPoint;
 
+		/** Defines how frame requests are handled. Whether they can block or not. */
+		EFrameLockingMode FrameLockingMode = EFrameLockingMode::FreeRun;
+
 		/** Whether the stream will output a frame at every frame interval, repeating last frame if no new one provided */
 		bool bDoContinuousOutput = true;
 
diff --git a/Engine/Plugins/VirtualProduction/Rivermax/RivermaxMedia/Source/RivermaxMedia/Private/RivermaxCustomTimeStep.cpp b/Engine/Plugins/VirtualProduction/Rivermax/RivermaxMedia/Source/RivermaxMedia/Private/RivermaxCustomTimeStep.cpp
--- a/Engine/Plugins/VirtualProduction/Rivermax/RivermaxMedia/Source/RivermaxMedia/Private/RivermaxCustomTimeStep.cpp
+++ b/Engine/Plugins/VirtualProduction/Rivermax/RivermaxMedia/Source/RivermaxMedia/Private/RivermaxCustomTimeStep.cpp
@@ -2,6 +2,7 @@
 
 #include "RivermaxCustomTimeStep.h"
 
+#include "IRivermaxBoundaryMonitor.h"
 #include "IRivermaxCoreModule.h"
 #include "IRivermaxManager.h"
 #include "RivermaxMediaLog.h"
@@ -47,6 +48,7 @@ void URivermaxCustomTimeStep::Shutdown(UEngine* InEngine)
 	IRivermaxCoreModule* RivermaxModule = FModuleManager::GetModulePtr<IRivermaxCoreModule>("RivermaxCore");
 	if (RivermaxModule)
 	{
+		DisableMonitoring();
 		RivermaxModule->GetRivermaxManager()->OnPostRivermaxManagerInit().RemoveAll(this);
 	}
 }
@@ -167,6 +169,9 @@ void URivermaxCustomTimeStep::PostEditChangeChainProperty(struct FPropertyChange
 	{
 		// If frame rate is changing live, reset frame number tracking to somewhat reinitialize our state
 		bIsPreviousFrameNumberValid = false;
+
+		// Cleanup and restart monitoring to match frame rate
+		EnableMonitoring();
 	}
 
 	Super::PostEditChangeChainProperty(PropertyChangedEvent);
@@ -207,8 +212,6 @@ bool URivermaxCustomTimeStep::WaitForNextFrame()
 		}
 	}
 
-	TRACE_BOOKMARK(TEXT("PTP genlock at %llu"), TargetTimeNanosec);
-
 	return ActualWaitTime < TimeoutSec;
 }
 
@@ -233,8 +236,39 @@ void URivermaxCustomTimeStep::OnRivermaxManagerInitialized()
 			return;
 		}
 
+		EnableMonitoring();
+
 		State = ECustomTimeStepSynchronizationState::Synchronized;
 	}
 }
 
+void URivermaxCustomTimeStep::EnableMonitoring()
+{
+	if (IRivermaxCoreModule* RivermaxModule = FModuleManager::GetModulePtr<IRivermaxCoreModule>("RivermaxCore"))
+	{
+		if (FrameRate != CachedFrameRate)
+		{
+			DisableMonitoring();
+
+			if (FrameRate.IsValid())
+			{
+				BoundaryMonitorListener = RivermaxModule->GetRivermaxBoundaryMonitor().StartMonitoring(FrameRate);
+				CachedFrameRate = FrameRate;
+			}
+		}
+	}
+}
+
+void URivermaxCustomTimeStep::DisableMonitoring()
+{
+	if (IRivermaxCoreModule* RivermaxModule = FModuleManager::GetModulePtr<IRivermaxCoreModule>("RivermaxCore"))
+	{
+		if (BoundaryMonitorListener.IsValid())
+		{
+			RivermaxModule->GetRivermaxBoundaryMonitor().StopMonitoring(BoundaryMonitorListener, CachedFrameRate);
+			BoundaryMonitorListener.Invalidate();
+		}
+	}
+}
+
 #undef LOCTEXT_NAMESPACE
diff --git a/Engine/Plugins/VirtualProduction/Rivermax/RivermaxMedia/Source/RivermaxMedia/Private/RivermaxMediaCapture.cpp b/Engine/Plugins/VirtualProduction/Rivermax/RivermaxMedia/Source/RivermaxMedia/Private/RivermaxMediaCapture.cpp
--- a/Engine/Plugins/VirtualProduction/Rivermax/RivermaxMedia/Source/RivermaxMedia/Private/RivermaxMediaCapture.cpp
+++ b/Engine/Plugins/VirtualProduction/Rivermax/RivermaxMedia/Source/RivermaxMedia/Private/RivermaxMediaCapture.cpp
@@ -2,6 +2,7 @@
 
 #include "RivermaxMediaCapture.h"
 
+#include "Async/Async.h"
 #include "Async/Fundamental/Task.h"
 #include "IRivermaxCoreModule.h"
 #include "IRivermaxManager.h"
@@ -20,42 +21,8 @@
 #endif
 
 
-DECLARE_GPU_STAT(Rivermax_Capture);
-DECLARE_GPU_STAT(Rivermax_SyncPointPass);
-
-
-static TAutoConsoleVariable<int32> CVarRivermaxPollTaskPriority(
-	TEXT("Rivermax.Output.PollTaskPriority"), static_cast<int32>(LowLevelTasks::ETaskPriority::High),
-	TEXT("Priority of the task responsible to poll the render fence"),
-	ECVF_Default);
-
-
-
-/** Structure holding data used for synchronization of buffer output */
-struct URivermaxMediaCapture::FRivermaxCaptureSyncData
-{
-	FRivermaxCaptureSyncData()
-		: bIsBusy(false)
-	{
-	}
-
-	/**
-	 * We use a RHIFence to write to in a pass following the buffer conversion one.
-	 * We spawn a task that will poll for the fence to be over (Waiting is not exposed)
-	 * and then we will push the captured buffer to rivermax stream
-	 */
-	FGPUFenceRHIRef RHIFence;
-
-	/** Whether fence is currently waiting to be cleared */
-	std::atomic<bool> bIsBusy;
-};
-
-
-
-/** Parameter to make our sync pass needing the convert pass as a prereq */
-BEGIN_SHADER_PARAMETER_STRUCT(FRivermaxSyncPassParameters, )
-	RDG_BUFFER_ACCESS(Buffer, ERHIAccess::CopySrc)
-END_SHADER_PARAMETER_STRUCT()
+DECLARE_GPU_STAT(Rmax_Capture);
+DECLARE_GPU_STAT(Rmax_FrameReservation);
 
 
 /* namespace RivermaxMediaCaptureDevice
@@ -234,7 +201,6 @@ bool URivermaxMediaCapture::Initialize(URivermaxMediaOutput* InMediaOutput)
 			RivermaxStream = Module->CreateOutputStream();
 			if (RivermaxStream)
 			{
-				bAreSyncHandlersInitialized = false;
 				bIsActive = RivermaxStream->Initialize(Options, *this);
 			}
 		}
@@ -243,49 +209,6 @@ bool URivermaxMediaCapture::Initialize(URivermaxMediaOutput* InMediaOutput)
 	return bIsActive;
 }
 
-void URivermaxMediaCapture::InitializeSyncHandlers_RenderThread()
-{
-	SyncHandlers.Reset(Options.NumberOfBuffers);
-	for (int32 Index = 0; Index < Options.NumberOfBuffers; ++Index)
-	{
-		TSharedPtr<FRivermaxCaptureSyncData> NewSync = MakeShared<FRivermaxCaptureSyncData>();
-		NewSync->RHIFence = RHICreateGPUFence(*FString::Printf(TEXT("RmaxSyncs_%02d"), Index));
-		SyncHandlers.Add(MoveTemp(NewSync));
-	}
-
-	bAreSyncHandlersInitialized = true;
-}
-
-TSharedPtr<URivermaxMediaCapture::FRivermaxCaptureSyncData> URivermaxMediaCapture::GetAvailableSyncHandler() const
-{
-	const auto FindAvailableHandlerFunc = [](const TSharedPtr<FRivermaxCaptureSyncData>& Item)
-	{
-		if (Item->bIsBusy == false)
-		{
-			return true;
-		}
-
-		return false;
-	};
-	
-	if (const TSharedPtr<FRivermaxCaptureSyncData>* FoundItem = SyncHandlers.FindByPredicate(FindAvailableHandlerFunc))
-	{
-		return *FoundItem;
-	}
-
-	return nullptr;
-}
-
-bool URivermaxMediaCapture::AreSyncHandlersBusy() const
-{
-	const auto IsBusyFunc = [](const TSharedPtr<FRivermaxCaptureSyncData>& Item)
-	{
-		return Item->bIsBusy.load();
-	};
-
-	return SyncHandlers.ContainsByPredicate(IsBusyFunc);
-}
-
 bool URivermaxMediaCapture::ConfigureStream(URivermaxMediaOutput* InMediaOutput, UE::RivermaxCore::FRivermaxOutputStreamOptions& OutOptions) const
 {
 	using namespace UE::RivermaxCore;
@@ -315,9 +238,10 @@ bool URivermaxMediaCapture::ConfigureStream(URivermaxMediaOutput* InMediaOutput,
 	}
 
 	OutOptions.FrameRate = InMediaOutput->FrameRate;
-	OutOptions.NumberOfBuffers = InMediaOutput->NumberOfTextureBuffers;
+	OutOptions.NumberOfBuffers = InMediaOutput->PresentationQueueSize;
 	OutOptions.bUseGPUDirect = InMediaOutput->bUseGPUDirect;
 	OutOptions.AlignmentMode = UE::RivermaxMediaUtils::Private::MediaOutputAlignmentToRivermaxAlignment(InMediaOutput->AlignmentMode);
+	OutOptions.FrameLockingMode = UE::RivermaxMediaUtils::Private::MediaOutputFrameLockingToRivermax(InMediaOutput->FrameLockingMode);
 
 	// Setup alignment dependent configs
 	OutOptions.bDoContinuousOutput = OutOptions.AlignmentMode == ERivermaxAlignmentMode::AlignmentPoint ? InMediaOutput->bDoContinuousOutput : false;
@@ -330,72 +254,24 @@ bool URivermaxMediaCapture::ConfigureStream(URivermaxMediaOutput* InMediaOutput,
 	return true;
 }
 
-void URivermaxMediaCapture::AddSyncPointPass(FRDGBuilder& GraphBuilder, const FCaptureBaseData& InBaseData, FRDGBufferRef OutputBuffer)
+void URivermaxMediaCapture::AddFrameReservationPass(FRDGBuilder& GraphBuilder)
 {
-	TRACE_CPUPROFILER_EVENT_SCOPE(RivermaxSyncPoint);
-	RDG_GPU_STAT_SCOPE(GraphBuilder, Rivermax_SyncPointPass);
-
-	// Initialize sync handlers only the first time. 
-	if (bAreSyncHandlersInitialized == false)
+	RDG_GPU_STAT_SCOPE(GraphBuilder, Rmax_FrameReservation);
+	SCOPED_DRAW_EVENT(GraphBuilder.RHICmdList, Rmax_FrameReservation);
+	
+	// Scene rendering will already be enqueued but capture conversion pass will not
+	// Revisit to push slot reservation till last minute
+	URivermaxMediaCapture* Capturer = this;
+	GraphBuilder.RHICmdList.EnqueueLambda([Capturer, FrameCounter = GFrameCounterRenderThread](FRHICommandList& RHICmdList)
 	{
-		InitializeSyncHandlers_RenderThread();
-	}
-
-	// Add buffer output as a parameter to depend on the compute shader pass
-	FRivermaxSyncPassParameters* PassParameters = GraphBuilder.AllocParameters<FRivermaxSyncPassParameters>();
-	PassParameters->Buffer = OutputBuffer;
-
-	// Prepare frame info we are going to push
-	UE::RivermaxCore::FRivermaxOutputVideoFrameInfo NewFrameInfo;
-	NewFrameInfo.FrameIdentifier = InBaseData.SourceFrameNumberRenderThread;
+		TRACE_CPUPROFILER_EVENT_SCOPE(RmaxFrameReservation);
+		TRACE_CPUPROFILER_EVENT_SCOPE_TEXT(*FString::Printf(TEXT("MediaCapturePipe: %llu"), FrameCounter));
 
-	URivermaxMediaCapture* Capturer = this;
-	GraphBuilder.AddPass(
-		RDG_EVENT_NAME("RivermaxCopySyncPass"),
-		PassParameters,
-		ERDGPassFlags::Copy | ERDGPassFlags::NeverCull,
-		[Capturer, OutputBuffer, NewFrameInfo](FRHICommandListImmediate& RHICmdList)
+		if (Capturer->bIsActive)
 		{
-			if (Capturer->bIsActive)
-			{
-				// Get available sync handler to create a sync point
-				TSharedPtr<FRivermaxCaptureSyncData> SyncDataPtr = Capturer->GetAvailableSyncHandler();
-				if (ensure(SyncDataPtr))
-				{
-					FRHIBuffer* RHIBuffer = OutputBuffer->GetRHI();
-
-					// This will happen after the compute shader conversion pass has completed
-					RHICmdList.WriteGPUFence(SyncDataPtr->RHIFence);
-					SyncDataPtr->bIsBusy = true;
-
-					// Spawn a task that will wait (poll) and continue the process of providing a new buffer to rivermax stream
-					UE::Tasks::Launch(UE_SOURCE_LOCATION, [Capturer, RHIBuffer, NewFrameInfo, SyncDataPtr]()
-					{	
-						TRACE_CPUPROFILER_EVENT_SCOPE(RmaxSyncTask);
-
-						// Wait until fence has been written (shader has completed)
-						while (SyncDataPtr->RHIFence->Poll() == false && Capturer->bIsActive)
-						{
-							constexpr float SleepTimeSeconds = 50 * 1E-6;
-							FPlatformProcess::SleepNoStats(SleepTimeSeconds);
-						}
-
-						SyncDataPtr->RHIFence->Clear();
-						SyncDataPtr->bIsBusy = false;
-
-						// Push new buffer to rivermax stream
-						if (Capturer->bIsActive && Capturer->RivermaxStream)
-						{
-							if (Capturer->RivermaxStream->PushGPUVideoFrame(NewFrameInfo, RHIBuffer) == false)
-							{
-								UE_LOG(LogRivermaxMedia, Verbose, TEXT("Failed to pushed captured frame"));
-							}
-						}
-
-					}, static_cast<LowLevelTasks::ETaskPriority>(CVarRivermaxPollTaskPriority.GetValueOnAnyThread()));
-				}
-			}
-		});
+			Capturer->RivermaxStream->ReserveFrame(FrameCounter);
+		}
+	});
 }
 
 void URivermaxMediaCapture::OnFrameCaptured_RenderingThread(const FCaptureBaseData& InBaseData, TSharedPtr<FMediaCaptureUserData, ESPMode::ThreadSafe> InUserData, void* InBuffer, int32 Width, int32 Height, int32 BytesPerRow)
@@ -418,7 +294,18 @@ void URivermaxMediaCapture::OnFrameCaptured_RenderingThread(const FCaptureBaseDa
 
 void URivermaxMediaCapture::OnRHIResourceCaptured_AnyThread(const FCaptureBaseData& InBaseData, TSharedPtr<FMediaCaptureUserData, ESPMode::ThreadSafe> InUserData, FBufferRHIRef InBuffer)
 {
-	// To implement when we start using media capture fence task
+	using namespace UE::RivermaxCore;
+
+	TRACE_CPUPROFILER_EVENT_SCOPE(URivermaxMediaCapture::OnRHIResourceCaptured_AnyThread);
+	TRACE_CPUPROFILER_EVENT_SCOPE_TEXT(*FString::Printf(TEXT("MediaCapturePipe: %u"), InBaseData.SourceFrameNumberRenderThread));
+
+	FRivermaxOutputVideoFrameInfo NewFrame;
+	NewFrame.FrameIdentifier = InBaseData.SourceFrameNumberRenderThread;
+	NewFrame.GPUBuffer = InBuffer;
+	if (RivermaxStream->PushVideoFrame(NewFrame) == false)
+	{
+		UE_LOG(LogRivermaxMedia, Verbose, TEXT("Failed to pushed captured frame"));
+	}
 }
 
 void URivermaxMediaCapture::OnFrameCaptured_AnyThread(const FCaptureBaseData& InBaseData, TSharedPtr<FMediaCaptureUserData, ESPMode::ThreadSafe> InUserData, const FMediaCaptureResourceData& InResourceData)
@@ -426,6 +313,7 @@ void URivermaxMediaCapture::OnFrameCaptured_AnyThread(const FCaptureBaseData& In
 	using namespace UE::RivermaxCore;
 
 	TRACE_CPUPROFILER_EVENT_SCOPE(URivermaxMediaCapture::OnFrameCaptured_AnyThread);
+	TRACE_CPUPROFILER_EVENT_SCOPE_TEXT(*FString::Printf(TEXT("MediaCapturePipe: %u"), InBaseData.SourceFrameNumberRenderThread));
 	
 	FRivermaxOutputVideoFrameInfo NewFrame;
 	NewFrame.Height = InResourceData.Height;
@@ -485,8 +373,10 @@ FRDGBufferDesc URivermaxMediaCapture::GetCustomBufferDescription(const FIntPoint
 
 void URivermaxMediaCapture::OnCustomCapture_RenderingThread(FRDGBuilder& GraphBuilder, const FCaptureBaseData& InBaseData, TSharedPtr<FMediaCaptureUserData, ESPMode::ThreadSafe> InUserData, FRDGTextureRef InSourceTexture, FRDGBufferRef OutputBuffer, const FRHICopyTextureInfo& CopyInfo, FVector2D CropU, FVector2D CropV)
 {
-	RDG_GPU_STAT_SCOPE(GraphBuilder, Rivermax_Capture)
+	RDG_GPU_STAT_SCOPE(GraphBuilder, Rmax_Capture)
 	TRACE_CPUPROFILER_EVENT_SCOPE(URivermaxMediaCapture::OnCustomCapture_RenderingThread);
+	TRACE_CPUPROFILER_EVENT_SCOPE_TEXT(*FString::Printf(TEXT("MediaCapturePipe: %llu"), GFrameCounterRenderThread));
+
 
 	using namespace UE::RivermaxShaders;
 	URivermaxMediaOutput* RivermaxOutput = CastChecked<URivermaxMediaOutput>(MediaOutput);
@@ -586,17 +476,13 @@ void URivermaxMediaCapture::OnCustomCapture_RenderingThread(FRDGBuilder& GraphBu
 		break;
 	}
 	}
-
-	// Adds a simple pass to have a fence written after compute shaders are done
-	if (RivermaxStream->IsGPUDirectSupported())
-	{
-		AddSyncPointPass(GraphBuilder, InBaseData, OutputBuffer);
-	}
+	
+	AddFrameReservationPass(GraphBuilder);
 }
 
 bool URivermaxMediaCapture::IsReadyForFinishDestroy()
 {
-	return Super::IsReadyForFinishDestroy() && !AreSyncHandlersBusy();
+	return Super::IsReadyForFinishDestroy();
 }
 
 void URivermaxMediaCapture::OnInitializationCompleted(bool bHasSucceed)
@@ -618,8 +504,9 @@ void URivermaxMediaCapture::OnStreamError()
 
 void URivermaxMediaCapture::OnPreFrameEnqueue()
 {
-	if (IsOutputSynchronizationSupported())
-	{
-		OnOutputSynchronization.ExecuteIfBound();
-	}
+	// Will need to add some logic in that callback chain for the case where margin wasn't enough
+	// For now, we act blindly that frames presented are all the same but we need a way to detect 
+	// if it's not and correct it.
+	TRACE_CPUPROFILER_EVENT_SCOPE(RmaxOutputSynchronization);
+	OnOutputSynchronization.ExecuteIfBound();
 }
diff --git a/Engine/Plugins/VirtualProduction/Rivermax/RivermaxMedia/Source/RivermaxMedia/Private/RivermaxMediaPlayer.cpp b/Engine/Plugins/VirtualProduction/Rivermax/RivermaxMedia/Source/RivermaxMedia/Private/RivermaxMediaPlayer.cpp
--- a/Engine/Plugins/VirtualProduction/Rivermax/RivermaxMedia/Source/RivermaxMedia/Private/RivermaxMediaPlayer.cpp
+++ b/Engine/Plugins/VirtualProduction/Rivermax/RivermaxMedia/Source/RivermaxMedia/Private/RivermaxMediaPlayer.cpp
@@ -32,6 +32,7 @@ DECLARE_CYCLE_STAT(TEXT("Rivermax MediaPlayer Request frame"), STAT_Rivermax_Med
 DECLARE_CYCLE_STAT(TEXT("Rivermax MediaPlayer Process frame"), STAT_Rivermax_MediaPlayer_ProcessFrame, STATGROUP_Media);
 
 DECLARE_GPU_STAT_NAMED(RivermaxMedia_SampleUsageFence, TEXT("RivermaxMedia_SampleUsageFence"));
+DECLARE_GPU_STAT_NAMED(Rmax_WaitForPixels, TEXT("Rmax_WaitForPixels"));
 
 
 namespace UE::RivermaxMedia
@@ -1133,9 +1134,17 @@ namespace UE::RivermaxMedia
 					}
 				);
 			}
-
+			
 			// Setup requirements for sample to be ready to be rendered
-			FRHICommandListImmediate & RHICmdList = FRHICommandListExecutor::GetImmediateCommandList();
+			FRHICommandListImmediate& RHICmdList = FRHICommandListExecutor::GetImmediateCommandList();
+
+			SCOPED_GPU_STAT(RHICmdList, Rmax_WaitForPixels);
+			SCOPED_DRAW_EVENT(RHICmdList, Rmax_WaitForPixels);
+
+			// Since we are going to enqueue a lambda that can potentially sleep in the RHI thread if the pixels haven't arrived,
+			// we dispatch the existing commands (including the draw event start timing in the SCOPED_DRAW_EVENT above) before any potential sleep.
+			RHICmdList.ImmediateFlush(EImmediateFlushType::DispatchToRHIThread);
+
 			RHICmdList.EnqueueLambda(
 				[NextFrameExpectations, this](FRHICommandList& RHICmdList)
 				{
diff --git a/Engine/Plugins/VirtualProduction/Rivermax/RivermaxMedia/Source/RivermaxMedia/Private/RivermaxMediaUtils.cpp b/Engine/Plugins/VirtualProduction/Rivermax/RivermaxMedia/Source/RivermaxMedia/Private/RivermaxMediaUtils.cpp
--- a/Engine/Plugins/VirtualProduction/Rivermax/RivermaxMedia/Source/RivermaxMedia/Private/RivermaxMediaUtils.cpp
+++ b/Engine/Plugins/VirtualProduction/Rivermax/RivermaxMedia/Source/RivermaxMedia/Private/RivermaxMediaUtils.cpp
@@ -243,4 +243,26 @@ namespace UE::RivermaxMediaUtils::Private
 		}
 	}
 
+	UE::RivermaxCore::EFrameLockingMode MediaOutputFrameLockingToRivermax(ERivermaxFrameLockingMode InFrameLockingMode)
+	{
+		using namespace UE::RivermaxCore;
+
+		switch (InFrameLockingMode)
+		{
+		case ERivermaxFrameLockingMode::FreeRun:
+		{
+			return EFrameLockingMode::FreeRun;
+		}
+		case ERivermaxFrameLockingMode::BlockOnReservation:
+		{
+			return EFrameLockingMode::BlockOnReservation;
+		}
+		default:
+		{
+			checkNoEntry();
+			return EFrameLockingMode::FreeRun;
+		}
+		}
+	}
+
 }
diff --git a/Engine/Plugins/VirtualProduction/Rivermax/RivermaxMedia/Source/RivermaxMedia/Private/RivermaxMediaUtils.h b/Engine/Plugins/VirtualProduction/Rivermax/RivermaxMedia/Source/RivermaxMedia/Private/RivermaxMediaUtils.h
--- a/Engine/Plugins/VirtualProduction/Rivermax/RivermaxMedia/Source/RivermaxMedia/Private/RivermaxMediaUtils.h
+++ b/Engine/Plugins/VirtualProduction/Rivermax/RivermaxMedia/Source/RivermaxMedia/Private/RivermaxMediaUtils.h
@@ -22,6 +22,7 @@ namespace UE::RivermaxMediaUtils::Private
 	ERivermaxMediaSourcePixelFormat RivermaxPixelFormatToMediaSourcePixelFormat(UE::RivermaxCore::ESamplingType InSamplingType);
 	ERivermaxMediaOutputPixelFormat RivermaxPixelFormatToMediaOutputPixelFormat(UE::RivermaxCore::ESamplingType InSamplingType);
 	UE::RivermaxCore::ERivermaxAlignmentMode MediaOutputAlignmentToRivermaxAlignment(ERivermaxMediaAlignmentMode InAlignmentMode);
+	UE::RivermaxCore::EFrameLockingMode MediaOutputFrameLockingToRivermax(ERivermaxFrameLockingMode InFrameLockingMode);
 	FSourceBufferDesc GetBufferDescription(const FIntPoint& Resolution, ERivermaxMediaSourcePixelFormat InPixelFormat);
 
 }
\ No newline at end of file
diff --git a/Engine/Plugins/VirtualProduction/Rivermax/RivermaxMedia/Source/RivermaxMedia/Public/RivermaxCustomTimeStep.h b/Engine/Plugins/VirtualProduction/Rivermax/RivermaxMedia/Source/RivermaxMedia/Public/RivermaxCustomTimeStep.h
--- a/Engine/Plugins/VirtualProduction/Rivermax/RivermaxMedia/Source/RivermaxMedia/Public/RivermaxCustomTimeStep.h
+++ b/Engine/Plugins/VirtualProduction/Rivermax/RivermaxMedia/Source/RivermaxMedia/Public/RivermaxCustomTimeStep.h
@@ -46,6 +46,12 @@ private:
 	/** We wait for Rivermax Manager to be initialized to do final verification before going to Synchronized state` */
 	void OnRivermaxManagerInitialized();
 
+	/** Adds current frame rate to boundary monitor */
+	void EnableMonitoring();
+
+	/** Removes current frame rate to boundary monitor */
+	void DisableMonitoring();
+
 public:
 
 	/** Target frame rate to which to genlock. Uses ST2059 standard to align PTP time to standard genlock */
@@ -73,4 +79,10 @@ private:
 	bool bIsPreviousFrameNumberValid = false;
 	uint64 PreviousFrameNumber = 0;
 	uint64 DeltaFrameNumber = 0;
+
+	/** Frame boundary monitoring ID */
+	FGuid BoundaryMonitorListener;
+
+	/** Cached frame rate used for boundary monitoring */
+	FFrameRate CachedFrameRate;
 };
diff --git a/Engine/Plugins/VirtualProduction/Rivermax/RivermaxMedia/Source/RivermaxMedia/Public/RivermaxMediaCapture.h b/Engine/Plugins/VirtualProduction/Rivermax/RivermaxMedia/Source/RivermaxMedia/Public/RivermaxMediaCapture.h
--- a/Engine/Plugins/VirtualProduction/Rivermax/RivermaxMedia/Source/RivermaxMedia/Public/RivermaxMediaCapture.h
+++ b/Engine/Plugins/VirtualProduction/Rivermax/RivermaxMedia/Source/RivermaxMedia/Public/RivermaxMediaCapture.h
@@ -73,21 +73,11 @@ private:
 	/** Initializes capture and launches stream creation */
 	bool Initialize(URivermaxMediaOutput* InMediaOutput);
 
-	/** Initializes sync handlers, fences, for gpudirect functionality */
-	void InitializeSyncHandlers_RenderThread();
-
-	/** Returns available sync handler. i.e fence not already used */
-	TSharedPtr<FRivermaxCaptureSyncData> GetAvailableSyncHandler() const;
-
-	/** Returns true if one of the sync handler is in flight waiting for its fence */
-	bool AreSyncHandlersBusy() const;
-
 	/** Configures rivermax stream with desired output options */
 	bool ConfigureStream(URivermaxMediaOutput* InMediaOutput, UE::RivermaxCore::FRivermaxOutputStreamOptions& OutOptions) const;
 
-	/** When GPUDirect is involved, we add a pass after the conversion pass to write a fence and know when we can consume the buffer */
-	void AddSyncPointPass(FRDGBuilder& GraphBuilder, const FCaptureBaseData& InBaseData, FRDGBufferRef OutputBuffer);
-
+	/** Enqueues a RHI lambda to reserve a spot for the next frame to capture */
+	void AddFrameReservationPass(FRDGBuilder& GraphBuilder);
 private:
 
 	/** Instance of the rivermax stream opened for this capture */
@@ -96,12 +86,6 @@ private:
 	/** Set of options used to configure output stream */
 	UE::RivermaxCore::FRivermaxOutputStreamOptions Options;
 
-	/** Whether sync handlers have been initialized or not.  */
-	bool bAreSyncHandlersInitialized = false;
-
 	/** Whether capture is active. Used to be queried by sync task */
 	std::atomic<bool> bIsActive;
-
-	/** Array of sync handlers (fence) to sync when captured buffer is completed */
-	TArray<TSharedPtr<FRivermaxCaptureSyncData>> SyncHandlers;
 };
diff --git a/Engine/Plugins/VirtualProduction/Rivermax/RivermaxMedia/Source/RivermaxMedia/Public/RivermaxMediaOutput.h b/Engine/Plugins/VirtualProduction/Rivermax/RivermaxMedia/Source/RivermaxMedia/Public/RivermaxMediaOutput.h
--- a/Engine/Plugins/VirtualProduction/Rivermax/RivermaxMedia/Source/RivermaxMedia/Public/RivermaxMediaOutput.h
+++ b/Engine/Plugins/VirtualProduction/Rivermax/RivermaxMedia/Source/RivermaxMedia/Public/RivermaxMediaOutput.h
@@ -36,6 +36,19 @@ enum class ERivermaxMediaAlignmentMode : uint8
 	FrameCreation,
 };
 
+/** 
+ * Controls how rivermax capture behaves when there are no buffers available to capture into
+ */
+UENUM()
+enum class ERivermaxFrameLockingMode : uint8
+{
+	/** If no frame available, continue */
+	FreeRun,
+
+	/** Blocks RHI thread prior to capture the current frame if no space is available. */
+	BlockOnReservation,
+};
+
 
 /**
  * Output information for a Rivermax media capture.
@@ -77,6 +90,14 @@ public:
 	UPROPERTY(EditAnywhere, BlueprintReadWrite, Category = "Output", meta = (EditCondition = "AlignmentMode != ERivermaxMediaAlignmentMode::FrameCreation"))
 	bool bDoContinuousOutput = true;
 
+	/** For alignment point mode, controls whether we stall engine before capturing if there are no buffer available to capture into */
+	UPROPERTY(EditAnywhere, BlueprintReadWrite, Category = "Output", meta = (EditCondition = "AlignmentMode != ERivermaxMediaAlignmentMode::FrameCreation"))
+	ERivermaxFrameLockingMode FrameLockingMode = ERivermaxFrameLockingMode::FreeRun;
+
+	/** Number of frames that can be queued / used in output queue. Frame being sent counts for 1. */
+	UPROPERTY(EditAnywhere, BlueprintReadWrite, Category = "Output", meta = (ClampMin = "2", ClampMax = "8", UIMin = "2", UIMax = "8"))
+	int32 PresentationQueueSize = 2;
+
 	/** 
 	 * Experimental flag to use frame counter instead of using NVIDIA Rivermax clock for timestamping output frames
 	 * Meant to be used for UE-UE streams where frame locking is done, e.g. nDisplay.
diff --git a/Engine/Source/Runtime/Core/Public/Misc/FrameRate.h b/Engine/Source/Runtime/Core/Public/Misc/FrameRate.h
--- a/Engine/Source/Runtime/Core/Public/Misc/FrameRate.h
+++ b/Engine/Source/Runtime/Core/Public/Misc/FrameRate.h
@@ -207,6 +207,11 @@ struct FFrameRate
 	{
 		return Rate.AsFrameTime(TimeInSeconds);
 	}
+
+	friend inline uint32 GetTypeHash(const FFrameRate& Rate)
+	{
+		return HashCombine(GetTypeHash(Rate.Numerator), GetTypeHash(Rate.Denominator));
+	}
 	
 	friend CORE_API FArchive& operator<<(FArchive& Ar, FFrameRate& FrameRate);
 
@@ -321,4 +326,4 @@ CORE_API TValueOrError<FFrameRate, FExpressionError> ParseFrameRate(const TCHAR*
 /**
  * Common Lex::TryParseString overload for FFrameRate
  */
-CORE_API bool TryParseString(FFrameRate& OutFrameRate, const TCHAR* InString);
\ No newline at end of file
+CORE_API bool TryParseString(FFrameRate& OutFrameRate, const TCHAR* InString);
