- Fixing clang warning. Horde issue 319320

diff --git a/Engine/Plugins/VirtualProduction/Rivermax/RivermaxCore/Source/RivermaxCore/Private/RivermaxFrameManager.cpp b/Engine/Plugins/VirtualProduction/Rivermax/RivermaxCore/Source/RivermaxCore/Private/RivermaxFrameManager.cpp
--- a/Engine/Plugins/VirtualProduction/Rivermax/RivermaxCore/Source/RivermaxCore/Private/RivermaxFrameManager.cpp
+++ b/Engine/Plugins/VirtualProduction/Rivermax/RivermaxCore/Source/RivermaxCore/Private/RivermaxFrameManager.cpp
@@ -100,7 +100,8 @@ namespace UE::RivermaxCore::Private
 		}
 
 		//Otherwise, prepare next free frame
-		if (NextFrame = GetFreeFrame())
+		NextFrame = GetFreeFrame();
+		if (NextFrame)
 		{
 			NextFrame->Reset();
 			NextFrame->FrameIdentifier = NextFrameIdentifier;
