- Adding critical error callback for rivermax frame manager
- Changing validity check in media capture render pass to avoid missing frames during gc

diff --git a/Engine/Plugins/Media/MediaIOFramework/Source/MediaIOCore/Private/MediaCaptureRenderPass.cpp b/Engine/Plugins/Media/MediaIOFramework/Source/MediaIOCore/Private/MediaCaptureRenderPass.cpp
--- a/Engine/Plugins/Media/MediaIOFramework/Source/MediaIOCore/Private/MediaCaptureRenderPass.cpp
+++ b/Engine/Plugins/Media/MediaIOFramework/Source/MediaIOCore/Private/MediaCaptureRenderPass.cpp
@@ -489,14 +489,14 @@ namespace UE::MediaCapture
 	}
 
 	FRenderPipeline::FRenderPipeline(UMediaCapture* InMediaCapture)
-		: WeakMediaCapture(InMediaCapture)
+		: MediaCapture(InMediaCapture)
 	{
 		RenderPasses = FRenderPass::CreateRenderPasses(InMediaCapture);
 	}
 
 	void FRenderPipeline::InitializeResources_RenderThread(const TSharedPtr<MediaCaptureData::FCaptureFrame>& CapturingFrame, FRDGBuilder& GraphBuilder, FRDGTextureRef InputResource) const
 	{
-		if (UMediaCapture* MediaCapture = WeakMediaCapture.Get())
+		if (CapturingFrame->bMediaCaptureActive)
 		{
 			FInitializePassOutputArgs InitializePassOutputArgs;
 			InitializePassOutputArgs.MediaCapture = MediaCapture;
@@ -525,7 +525,7 @@ namespace UE::MediaCapture
 	{
 		FRDGViewableResource* PreviousOutput = InputRGBTexture;
 
-		if (UMediaCapture* MediaCapture = WeakMediaCapture.Get())
+		if (CapturingFrame->bMediaCaptureActive)
 		{
 			for (const UE::MediaCapture::FRenderPass& RenderPass : RenderPasses)
 			{
diff --git a/Engine/Plugins/Media/MediaIOFramework/Source/MediaIOCore/Private/MediaCaptureRenderPass.h b/Engine/Plugins/Media/MediaIOFramework/Source/MediaIOCore/Private/MediaCaptureRenderPass.h
--- a/Engine/Plugins/Media/MediaIOFramework/Source/MediaIOCore/Private/MediaCaptureRenderPass.h
+++ b/Engine/Plugins/Media/MediaIOFramework/Source/MediaIOCore/Private/MediaCaptureRenderPass.h
@@ -188,7 +188,7 @@ namespace UE::MediaCapture
 		/** Register a render pass' resources with the render graph builder. Calling this multiple times will simply return the (already) registered resource. */
 		static FRDGViewableResource* RegisterPassOutputResource(const UE::MediaCapture::FRenderPassFrameResources& InFrameResources, FRDGBuilder& InRDGBuilder, const UE::MediaCapture::FRenderPass& InRenderPass);
 	private:
-		TWeakObjectPtr<UMediaCapture> WeakMediaCapture;
+		UMediaCapture* MediaCapture = nullptr;
 		TArray<FRenderPass> RenderPasses;
 	};
 
diff --git a/Engine/Plugins/VirtualProduction/Rivermax/RivermaxCore/Source/RivermaxCore/Private/RivermaxFrameManager.cpp b/Engine/Plugins/VirtualProduction/Rivermax/RivermaxCore/Source/RivermaxCore/Private/RivermaxFrameManager.cpp
--- a/Engine/Plugins/VirtualProduction/Rivermax/RivermaxCore/Source/RivermaxCore/Private/RivermaxFrameManager.cpp
+++ b/Engine/Plugins/VirtualProduction/Rivermax/RivermaxCore/Source/RivermaxCore/Private/RivermaxFrameManager.cpp
@@ -21,6 +21,7 @@ namespace UE::RivermaxCore::Private
 		OnFrameReadyDelegate = Args.OnFrameReadyDelegate;
 		OnPreFrameReadyDelegate = Args.OnPreFrameReadyDelegate;
 		OnFreeFrameDelegate = Args.OnFreeFrameDelegate;
+		OnCriticalErrorDelegate = Args.OnCriticalErrorDelegate;
 		FrameResolution = Args.Resolution;
 		TotalFrameCount = Args.NumberOfFrames;
 		const uint32 FrameSize = FrameResolution.Y * Args.Stride;
@@ -190,11 +191,17 @@ namespace UE::RivermaxCore::Private
 
 	bool FFrameManager::SetFrameData(const FRivermaxOutputVideoFrameInfo& NewFrameInfo)
 	{
+		bool bSuccess = false;
 		if (TSharedPtr<FRivermaxOutputFrame> NextFrame = GetNextFrame(NewFrameInfo.FrameIdentifier))
 		{
-			return FrameAllocator->CopyData(NewFrameInfo, NextFrame);
+			 bSuccess = FrameAllocator->CopyData(NewFrameInfo, NextFrame);
+			 if (!bSuccess)
+			 {
+				 OnCriticalErrorDelegate.ExecuteIfBound();
+			 }
+			 
 		}
-		return true;
+		return bSuccess;
 	}
 
 	void FFrameManager::OnDataCopied(const TSharedPtr<FRivermaxOutputFrame>& Frame)
diff --git a/Engine/Plugins/VirtualProduction/Rivermax/RivermaxCore/Source/RivermaxCore/Private/RivermaxFrameManager.h b/Engine/Plugins/VirtualProduction/Rivermax/RivermaxCore/Source/RivermaxCore/Private/RivermaxFrameManager.h
--- a/Engine/Plugins/VirtualProduction/Rivermax/RivermaxCore/Source/RivermaxCore/Private/RivermaxFrameManager.h
+++ b/Engine/Plugins/VirtualProduction/Rivermax/RivermaxCore/Source/RivermaxCore/Private/RivermaxFrameManager.h
@@ -35,6 +35,7 @@ namespace UE::RivermaxCore::Private
 	DECLARE_DELEGATE(FOnFrameReadyDelegate);
 	DECLARE_DELEGATE(FOnPreFrameReadyDelegate);
 	DECLARE_DELEGATE(FOnFreeFrameDelegate);
+	DECLARE_DELEGATE(FOnCriticalErrorDelegate);
 
 	/** Holds arguments to configure frame manager during initialization */
 	struct FFrameManagerSetupArgs
@@ -52,13 +53,16 @@ namespace UE::RivermaxCore::Private
 		bool bTryGPUAllocation = true;
 
 		/** Delegate called when a frame is now free to use */
-		FOnFreeFrameDelegate OnFreeFrameDelegate = nullptr;
+		FOnFreeFrameDelegate OnFreeFrameDelegate;
 
 		/** Delegate triggered just before a frame is enqueued to be sent */
-		FOnPreFrameReadyDelegate OnPreFrameReadyDelegate = nullptr;
+		FOnPreFrameReadyDelegate OnPreFrameReadyDelegate;
 		
 		/** Delegate called when a frame is now ready to be sent */
-		FOnFrameReadyDelegate OnFrameReadyDelegate = nullptr;
+		FOnFrameReadyDelegate OnFrameReadyDelegate;
+
+		/** Delegate called when a critical has happened and stream should shut down */
+		FOnCriticalErrorDelegate OnCriticalErrorDelegate;
 	};
 
 	/** 
@@ -150,13 +154,16 @@ namespace UE::RivermaxCore::Private
 		TUniquePtr<FBaseFrameAllocator> FrameAllocator;
 
 		/** Delegate triggered when a frame is free to use */
-		FOnFreeFrameDelegate OnFreeFrameDelegate = nullptr;
+		FOnFreeFrameDelegate OnFreeFrameDelegate;
 
 		/** Delegate triggered just before a frame is enqueued to be sent */
-		FOnPreFrameReadyDelegate OnPreFrameReadyDelegate = nullptr;
+		FOnPreFrameReadyDelegate OnPreFrameReadyDelegate;
 		
 		/** Delegate triggered when a frame is ready to be sent (video data has been copied) */
-		FOnFrameReadyDelegate OnFrameReadyDelegate = nullptr;
+		FOnFrameReadyDelegate OnFrameReadyDelegate;
+
+		/** Delegate triggered when a critical error has happened and stream should shut down */
+		FOnCriticalErrorDelegate OnCriticalErrorDelegate;
 		
 		/** List of available frames. */
 		TArray<int32> FreeFrames;
diff --git a/Engine/Plugins/VirtualProduction/Rivermax/RivermaxCore/Source/RivermaxCore/Private/RivermaxOutputStream.cpp b/Engine/Plugins/VirtualProduction/Rivermax/RivermaxCore/Source/RivermaxCore/Private/RivermaxOutputStream.cpp
--- a/Engine/Plugins/VirtualProduction/Rivermax/RivermaxCore/Source/RivermaxCore/Private/RivermaxOutputStream.cpp
+++ b/Engine/Plugins/VirtualProduction/Rivermax/RivermaxCore/Source/RivermaxCore/Private/RivermaxOutputStream.cpp
@@ -412,13 +412,7 @@ namespace UE::RivermaxCore::Private
 	{
 		TRACE_CPUPROFILER_EVENT_SCOPE(FRivermaxOutputStream::PushVideoFrame);
 
-		const bool bSuccess = FrameManager->SetFrameData(NewFrame);
-		if (!bSuccess)
-		{
-			Listener->OnStreamError();
-		}
-
-		return bSuccess;
+		return FrameManager->SetFrameData(NewFrame);
 	}
 
 	void FRivermaxOutputStream::Process_AnyThread()
@@ -1322,5 +1316,12 @@ namespace UE::RivermaxCore::Private
 	{
 		Listener->OnPreFrameEnqueue();
 	}
+
+	void FRivermaxOutputStream::OnFrameManagerCriticalError()
+	{
+		Listener->OnStreamError();
+		Stop();
+	}
+
 }
 
diff --git a/Engine/Plugins/VirtualProduction/Rivermax/RivermaxCore/Source/RivermaxCore/Private/RivermaxOutputStream.h b/Engine/Plugins/VirtualProduction/Rivermax/RivermaxCore/Source/RivermaxCore/Private/RivermaxOutputStream.h
--- a/Engine/Plugins/VirtualProduction/Rivermax/RivermaxCore/Source/RivermaxCore/Private/RivermaxOutputStream.h
+++ b/Engine/Plugins/VirtualProduction/Rivermax/RivermaxCore/Source/RivermaxCore/Private/RivermaxOutputStream.h
@@ -180,6 +180,10 @@ namespace UE::RivermaxCore::Private
 		/** Used to know when a frame is ready to be used and receive new data */
 		void OnFrameReadyToBeUsed();
 
+		/** Used to detect when the frame manager has caught a critical error */
+		void OnFrameManagerCriticalError();
+
+
 	private:
 
 		/** Options related to this stream. i.e resolution, frame rate, etc... */
