- Fixing crash in media capture when shutting down or closing editor

diff --git a/Engine/Plugins/Media/MediaIOFramework/Source/MediaIOCore/Private/MediaCaptureSyncPointWatcher.cpp b/Engine/Plugins/Media/MediaIOFramework/Source/MediaIOCore/Private/MediaCaptureSyncPointWatcher.cpp
--- a/Engine/Plugins/Media/MediaIOFramework/Source/MediaIOCore/Private/MediaCaptureSyncPointWatcher.cpp
+++ b/Engine/Plugins/Media/MediaIOFramework/Source/MediaIOCore/Private/MediaCaptureSyncPointWatcher.cpp
@@ -138,8 +138,9 @@ void FSyncPointWatcher::WaitForSinglePendingTaskToComplete()
 
 FSyncPointWatcher::~FSyncPointWatcher()
 {
-	Stop();
-	WorkingThread->WaitForCompletion();
+	constexpr bool bWaitForThread = true;
+	WorkingThread->Kill(bWaitForThread);
+	WorkingThread.Reset();
 }
 
 }
\ No newline at end of file
diff --git a/Engine/Plugins/VirtualProduction/Rivermax/RivermaxCore/Source/RivermaxCore/Private/RivermaxOutputStream.cpp b/Engine/Plugins/VirtualProduction/Rivermax/RivermaxCore/Source/RivermaxCore/Private/RivermaxOutputStream.cpp
--- a/Engine/Plugins/VirtualProduction/Rivermax/RivermaxCore/Source/RivermaxCore/Private/RivermaxOutputStream.cpp
+++ b/Engine/Plugins/VirtualProduction/Rivermax/RivermaxCore/Source/RivermaxCore/Private/RivermaxOutputStream.cpp
@@ -329,7 +329,7 @@ namespace UE::RivermaxCore::Private
 		// Enable frame boundary monitoring
 		MonitoringGuid = RivermaxModule->GetRivermaxBoundaryMonitor().StartMonitoring(Options.FrameRate);
 
-		Async(EAsyncExecution::TaskGraph, [this]()
+		InitializationFuture = Async(EAsyncExecution::TaskGraph, [this]()
 		{
 			TAnsiStringBuilder<2048> SDPDescription;
 			UE::RivermaxCore::Private::Utils::StreamOptionsToSDPDescription(Options, SDPDescription);
@@ -400,6 +400,11 @@ namespace UE::RivermaxCore::Private
 
 	void FRivermaxOutputStream::Uninitialize()
 	{
+		if (!InitializationFuture.IsReady())
+		{
+			InitializationFuture.Wait();
+		}
+
 		if (RivermaxThread != nullptr)
 		{
 			Stop();
diff --git a/Engine/Plugins/VirtualProduction/Rivermax/RivermaxCore/Source/RivermaxCore/Private/RivermaxOutputStream.h b/Engine/Plugins/VirtualProduction/Rivermax/RivermaxCore/Source/RivermaxCore/Private/RivermaxOutputStream.h
--- a/Engine/Plugins/VirtualProduction/Rivermax/RivermaxCore/Source/RivermaxCore/Private/RivermaxOutputStream.h
+++ b/Engine/Plugins/VirtualProduction/Rivermax/RivermaxCore/Source/RivermaxCore/Private/RivermaxOutputStream.h
@@ -4,6 +4,7 @@
 
 #include "IRivermaxOutputStream.h"
 
+#include "Async/Future.h"
 #include "Containers/SpscQueue.h"
 #include "HAL/Runnable.h"
 #include "RivermaxHeader.h"
@@ -246,6 +247,9 @@ namespace UE::RivermaxCore::Private
 		/** Guid given by boundary monitoring handler to unregister ourselves */
 		FGuid MonitoringGuid;
 
+		/** Future returned by the async initialization job we launch. Used to detect if it has completed during shutdown. */
+		TFuture<void> InitializationFuture;
+
 		/** Time to sleep when waiting for an operation to complete */
 		static constexpr double SleepTimeSeconds = 50.0 * 1E-6;
 	};
