- Fixing rmax boundary monitor not removing old monitored frame rate

diff --git a/Engine/Plugins/VirtualProduction/Rivermax/RivermaxCore/Source/RivermaxCore/Private/RivermaxBoundaryMonitor.cpp b/Engine/Plugins/VirtualProduction/Rivermax/RivermaxCore/Source/RivermaxCore/Private/RivermaxBoundaryMonitor.cpp
--- a/Engine/Plugins/VirtualProduction/Rivermax/RivermaxCore/Source/RivermaxCore/Private/RivermaxBoundaryMonitor.cpp
+++ b/Engine/Plugins/VirtualProduction/Rivermax/RivermaxCore/Source/RivermaxCore/Private/RivermaxBoundaryMonitor.cpp
@@ -97,7 +97,7 @@ namespace UE::RivermaxCore
 	FBoundaryMonitor::~FBoundaryMonitor()
 	{
 		bIsEnabled = false;
-		WorkingThread->WaitForCompletion();
+		WorkingThread->Kill();
 	}
 
 	uint32 FBoundaryMonitor::Run()
@@ -209,6 +209,12 @@ void FRivermaxBoundaryMonitor::StartMonitor(const FFrameRate& FrameRate)
 
 void FRivermaxBoundaryMonitor::StopMonitor(const FFrameRate& FrameRate)
 {
+	TUniquePtr<FBoundaryMonitor>* Monitor = MonitorMap.Find(FrameRate);
+	if (Monitor)
+	{
+		(*Monitor).Reset();
+	}
+
 	MonitorMap.Remove(FrameRate);
 }
 
