- Adding support for intermediate buffer usage in rivermax memory setup to avoid using skip chunk method when repeating a frame causing timing issues.
- Adding single rivermax memblock support which can potentially improve perf
- Adding cvars to control memory setup and slice count when memcopying a frame into intermediate buffer. Automatic selection using frame descriptor as heuristic not done yet.

- Fixing RTP header pointer management that could endup writing out of bounds
- Fixing frame reservation not being done when waiting for one possibly causing a skipped frame

diff --git a/Engine/Plugins/VirtualProduction/Rivermax/RivermaxCore/Source/RivermaxCore/Private/RivermaxFrameAllocator.cpp b/Engine/Plugins/VirtualProduction/Rivermax/RivermaxCore/Source/RivermaxCore/Private/RivermaxFrameAllocator.cpp
--- a/Engine/Plugins/VirtualProduction/Rivermax/RivermaxCore/Source/RivermaxCore/Private/RivermaxFrameAllocator.cpp
+++ b/Engine/Plugins/VirtualProduction/Rivermax/RivermaxCore/Source/RivermaxCore/Private/RivermaxFrameAllocator.cpp
@@ -2,6 +2,7 @@
 
 #include "RivermaxFrameAllocator.h"
 
+#include "Async/ParallelFor.h"
 #include "CudaModule.h"
 #include "RivermaxLog.h"
 #include "ID3D12DynamicRHI.h"
@@ -9,7 +10,23 @@
 
 namespace UE::RivermaxCore::Private
 {
-	bool FGPUAllocator::Allocate(int32 FrameCount)
+	static TAutoConsoleVariable<int32> CVarRivermaxOutputEnableParallelCopy(
+		TEXT("Rivermax.Output.EnableParallelCopy"), 0,
+		TEXT("Parallel"),
+		ECVF_Default);
+
+	static TAutoConsoleVariable<int32> CVarRivermaxOutputParallelCopyThreadCount(
+		TEXT("Rivermax.Output.ParallelCopyThreadCount"), 8,
+		TEXT("Parallel"),
+		ECVF_Default);
+
+	FGPUAllocator::FGPUAllocator(int32 InDesiredFrameSize, FOnFrameDataCopiedDelegate InOnDataCopiedDelegate)
+		: Super(InDesiredFrameSize, InOnDataCopiedDelegate)
+	{
+
+	}
+
+	bool FGPUAllocator::Allocate(int32 FrameCount, bool bAlignFrameMemory)
 	{
 		// Allocate a single memory space that will contain all frame buffers
 		TRACE_CPUPROFILER_EVENT_SCOPE(Rmax::GPUAllocation);
@@ -126,13 +143,20 @@ namespace UE::RivermaxCore::Private
 		}
 
 		AllocatedFrameCount = FrameCount;
-		AllocatedFrameSize = CudaAlignedFrameSize;
+
+		if (bAlignFrameMemory)
+		{
+			AllocatedFrameSize = CudaAlignedFrameSize;
+		}
+		else
+		{
+			AllocatedFrameSize = DesiredFrameSize;
+		}
 
 		for (int32 Index = 0; Index < AllocatedFrameCount; ++Index)
 		{
 			TSharedPtr<FRivermaxOutputFrame> Frame = MakeShared<FRivermaxOutputFrame>(Index);
 			Frame->VideoBuffer = GetFrameAddress(Index);
-			Frames.Add(MoveTemp(Frame));
 		}
 
 		CudaModule.DriverAPI()->cuCtxPopCurrent(nullptr);
@@ -176,7 +200,6 @@ namespace UE::RivermaxCore::Private
 			}
 			GPUStream = nullptr;
 			AllocatedFrameCount = 0;
-			Frames.Empty();
 
 			CudaModule.DriverAPI()->cuCtxPopCurrent(nullptr);
 		}
@@ -192,52 +215,75 @@ namespace UE::RivermaxCore::Private
 		return nullptr;
 	}
 
-	bool FGPUAllocator::CopyData(const FRivermaxOutputVideoFrameInfo& NewFrame, const TSharedPtr<FRivermaxOutputFrame>& DestinationFrame)
+	bool FGPUAllocator::CopyData(const FCopyArgs& Args)
 	{
-		TRACE_CPUPROFILER_EVENT_SCOPE(Rmax::GPUCopyStart);
-		check(NewFrame.GPUBuffer);
+		check(Args.RHISourceMemory || Args.SourceMemory);
+
+		TRACE_CPUPROFILER_EVENT_SCOPE(Rmax::CudaCopyStart);
 		{
 			FCUDAModule& CudaModule = FModuleManager::GetModuleChecked<FCUDAModule>("CUDA");
 			CUresult Result = CudaModule.DriverAPI()->cuCtxPushCurrent(CudaModule.GetCudaContext());
 
-			void* MappedPointer = GetMappedAddress(NewFrame.GPUBuffer);
-			if (MappedPointer == nullptr)
+			ON_SCOPE_EXIT
 			{
-				UE_LOG(LogRivermax, Error, TEXT("Failed to find a mapped memory address for captured buffer. Stopping capture."));
-				return false;
+				FCUDAModule::CUDA().cuCtxPopCurrent(nullptr); 
+			};
+
+			// If source memory is rhi memory, we need to go through mapping, otherwise, it's already cuda located
+			void* SourceMemory = Args.SourceMemory;
+			if (SourceMemory == nullptr)
+			{
+				SourceMemory = GetMappedAddress(Args.RHISourceMemory);
+				if (SourceMemory == nullptr)
+				{
+					UE_LOG(LogRivermax, Error, TEXT("Failed to find a mapped memory address for captured buffer. Stopping capture."));
+					return false;
+				}
 			}
 
-			const CUdeviceptr CudaMemoryPointer = reinterpret_cast<CUdeviceptr>(MappedPointer);
-			Result = CudaModule.DriverAPI()->cuMemcpyDtoDAsync(reinterpret_cast<CUdeviceptr>(DestinationFrame->VideoBuffer), CudaMemoryPointer, DesiredFrameSize, reinterpret_cast<CUstream>(GPUStream));
-			if (Result != CUDA_SUCCESS)
+			const CUdeviceptr DestinationCudaPointer = reinterpret_cast<CUdeviceptr>(Args.DestinationMemory);
+			const CUdeviceptr BaseCudaPointer = reinterpret_cast<CUdeviceptr>(CudaAllocatedMemoryBaseAddress);
+			if(DestinationCudaPointer < BaseCudaPointer || DestinationCudaPointer >= (BaseCudaPointer + CudaAllocatedMemory))
 			{
-				UE_LOG(LogRivermax, Error, TEXT("Failed to copy captured bufer to cuda memory. Stopping capture. Error: %d"), Result);
+				UE_LOG(LogRivermax, Error, TEXT("Trying to copy data outside of allocated buffer."));
 				return false;
 			}
 
+			{
+				TRACE_CPUPROFILER_EVENT_SCOPE(Rmax::CudaMemcopyAsync);
+				const CUdeviceptr CudaMemoryPointer = reinterpret_cast<CUdeviceptr>(SourceMemory);
+				Result = CudaModule.DriverAPI()->cuMemcpyDtoDAsync(reinterpret_cast<CUdeviceptr>(Args.DestinationMemory), CudaMemoryPointer, Args.SizeToCopy, reinterpret_cast<CUstream>(GPUStream));
+				if (Result != CUDA_SUCCESS)
+				{
+					UE_LOG(LogRivermax, Error, TEXT("Failed to copy captured buffer to cuda memory. Stopping capture. Error: %d"), Result);
+					return false;
+				}
+			}
+			
+
 			// Callback called by Cuda when stream work has completed on cuda engine (MemCpy -> Callback)
 			// Once Memcpy has been done, we know we can mark that memory as available to be sent. 
 
-			auto CudaCallback = [](void* userData)
+			auto CudaCallback = [](void* UserData)
 			{
-				TRACE_CPUPROFILER_EVENT_SCOPE(Rmax::GPUCopyDone);
+				TRACE_CPUPROFILER_EVENT_SCOPE(Rmax::CudaCopyDone);
 
-				FGPUAllocator* Allocator = reinterpret_cast<FGPUAllocator*>(userData);
-				TOptional<TSharedPtr<FRivermaxOutputFrame>> FrameCopied = Allocator->PendingCopies.Dequeue();
-				if (FrameCopied.IsSet())
+				FGPUAllocator* Allocator = reinterpret_cast<FGPUAllocator*>(UserData);
+				TOptional<TSharedPtr<FBaseDataCopySideCar>> Sidecar = Allocator->PendingCopies.Dequeue();
+				if (Sidecar.IsSet())
 				{
-					FrameCopied.GetValue()->bIsVideoBufferReady = true;
-					Allocator->OnDataCopiedDelegate.ExecuteIfBound(FrameCopied.GetValue());
+					Allocator->OnDataCopiedDelegate.ExecuteIfBound(Sidecar.GetValue());
 				}
 			};
 
-			// Add pending frame for cuda callback 
-			PendingCopies.Enqueue(DestinationFrame);
+			// Add pending payload for cuda callback 
+			PendingCopies.Enqueue(Args.SideCar);
 
 			// Schedule a callback to make the frame available
-			CudaModule.DriverAPI()->cuLaunchHostFunc(reinterpret_cast<CUstream>(GPUStream), CudaCallback, this);
-
-			FCUDAModule::CUDA().cuCtxPopCurrent(nullptr);
+			{
+				TRACE_CPUPROFILER_EVENT_SCOPE(Rmax::CudaLaunchHost);
+				CudaModule.DriverAPI()->cuLaunchHostFunc(reinterpret_cast<CUstream>(GPUStream), CudaCallback, this);
+			}
 
 			return true;
 		}
@@ -245,12 +291,6 @@ namespace UE::RivermaxCore::Private
 		return false;
 	}
 
-	FGPUAllocator::FGPUAllocator(FIntPoint InFrameResolution, uint32 InStride, FOnFrameDataCopiedDelegate InOnDataCopiedDelegate)
-		: Super(InFrameResolution, InStride, InOnDataCopiedDelegate)
-	{
-
-	}
-
 	void* FGPUAllocator::GetMappedAddress(const FBufferRHIRef& InBuffer)
 	{
 		// If we are here, d3d12 had to have been validated
@@ -332,27 +372,18 @@ namespace UE::RivermaxCore::Private
 		return BufferCudaMemoryMap[InBuffer];
 	}
 
-	bool FSystemAllocator::Allocate(int32 FrameCount)
+	bool FSystemAllocator::Allocate(int32 FrameCount, bool bAlignFrameMemory)
 	{
 		TRACE_CPUPROFILER_EVENT_SCOPE(Rmax::SystemAllocation);
 
 		constexpr uint32 CacheLineSize = PLATFORM_CACHE_LINE_SIZE;
 
-		AllocatedFrameSize = Align(DesiredFrameSize, CacheLineSize);
-		const size_t NeededSystemMemory = Align(AllocatedFrameSize * FrameCount, CacheLineSize);
+		AllocatedFrameSize = bAlignFrameMemory ? Align(DesiredFrameSize, CacheLineSize) : DesiredFrameSize;
 		SystemMemoryBaseAddress = FMemory::Malloc(AllocatedFrameSize * FrameCount, CacheLineSize);
 
 		if (SystemMemoryBaseAddress)
 		{
 			AllocatedFrameCount = FrameCount;
-
-			for (int32 Index = 0; Index < AllocatedFrameCount; ++Index)
-			{
-				TSharedPtr<FRivermaxOutputFrame> Frame = MakeShared<FRivermaxOutputFrame>(Index);
-				Frame->VideoBuffer = GetFrameAddress(Index);
-				Frames.Add(MoveTemp(Frame));
-			}
-
 			return true;
 		}
 
@@ -366,7 +397,6 @@ namespace UE::RivermaxCore::Private
 			FMemory::Free(SystemMemoryBaseAddress);
 			SystemMemoryBaseAddress = nullptr;
 			AllocatedFrameCount = 0;
-			Frames.Empty();
 		}
 	}
 
@@ -380,38 +410,52 @@ namespace UE::RivermaxCore::Private
 		return nullptr;
 	}
 
-	FSystemAllocator::FSystemAllocator(FIntPoint InFrameResolution, uint32 InStride, FOnFrameDataCopiedDelegate InOnDataCopiedDelegate)
-		: Super(InFrameResolution, InStride, InOnDataCopiedDelegate)
+	FSystemAllocator::FSystemAllocator(int32 InDesiredFrameSize, FOnFrameDataCopiedDelegate InOnDataCopiedDelegate)
+		: Super(InDesiredFrameSize, InOnDataCopiedDelegate)
 	{
-
+		bUseParallelFor = CVarRivermaxOutputEnableParallelCopy.GetValueOnAnyThread() != 0;
+		ParallelForThreadCount = FMath::Clamp(CVarRivermaxOutputParallelCopyThreadCount.GetValueOnAnyThread(), 1, 100);
 	}
 
-	bool FSystemAllocator::CopyData(const FRivermaxOutputVideoFrameInfo& NewFrame, const TSharedPtr<FRivermaxOutputFrame>& DestinationFrame)
+	bool FSystemAllocator::CopyData(const FCopyArgs& Args)
 	{
-		FMemory::Memcpy(DestinationFrame->VideoBuffer, NewFrame.VideoBuffer, DesiredFrameSize);
+		// GPU memory copy shouldn't happen for system allocator
+		if (Args.RHISourceMemory != nullptr)
+		{
+			ensure(false);
+			return false;
+		}
 
-		DestinationFrame->bIsVideoBufferReady = true;
-		OnDataCopiedDelegate.ExecuteIfBound(DestinationFrame);
+
+		if(bUseParallelFor)
+		{
+			const int32 ThreadCount = ParallelForThreadCount;
+			const uint32 SizePerThread = 1 + ((Args.SizeToCopy - 1) / ThreadCount);
+			ParallelFor(ThreadCount, [SizePerThread, &Args](int32 BlockIndex)
+			{
+				const uint32 Offset = BlockIndex * SizePerThread;
+				const uint32 MaxSize = Args.SizeToCopy - Offset;
+				const uint32 CopySize = FMath::Min(SizePerThread, MaxSize);
+				uint8* SourceAddress = reinterpret_cast<uint8*>(Args.SourceMemory) + Offset;
+				uint8* DestinationAddress = reinterpret_cast<uint8*>(Args.DestinationMemory) + Offset;
+				FMemory::Memcpy(DestinationAddress, SourceAddress, CopySize);
+			});
+		}
+		else
+		{
+			FMemory::Memcpy(Args.DestinationMemory, Args.SourceMemory, Args.SizeToCopy);
+		}
+		
+		OnDataCopiedDelegate.ExecuteIfBound(Args.SideCar);
 
 		return true;
 	}
 
-	FBaseFrameAllocator::FBaseFrameAllocator(FIntPoint InFrameResolution, uint32 InStride, FOnFrameDataCopiedDelegate InOnDataCopiedDelegate)
-		: FrameResolution(InFrameResolution)
-		, Stride(InStride)
-		, DesiredFrameSize(FrameResolution.Y * Stride)
+	FBaseFrameAllocator::FBaseFrameAllocator(int32 InDesiredFrameSize, FOnFrameDataCopiedDelegate InOnDataCopiedDelegate)
+		: DesiredFrameSize(InDesiredFrameSize)
 		, OnDataCopiedDelegate(MoveTemp(InOnDataCopiedDelegate))
 	{
 
 	}
-
-	const TSharedPtr<FRivermaxOutputFrame> FBaseFrameAllocator::GetFrame(int32 Index) const
-	{
-		if (Frames.IsValidIndex(Index))
-		{
-			return Frames[Index];
-		}
-		return nullptr;
-	}
 }
 
diff --git a/Engine/Plugins/VirtualProduction/Rivermax/RivermaxCore/Source/RivermaxCore/Private/RivermaxFrameAllocator.h b/Engine/Plugins/VirtualProduction/Rivermax/RivermaxCore/Source/RivermaxCore/Private/RivermaxFrameAllocator.h
--- a/Engine/Plugins/VirtualProduction/Rivermax/RivermaxCore/Source/RivermaxCore/Private/RivermaxFrameAllocator.h
+++ b/Engine/Plugins/VirtualProduction/Rivermax/RivermaxCore/Source/RivermaxCore/Private/RivermaxFrameAllocator.h
@@ -17,46 +17,60 @@ namespace UE::RivermaxCore
 
 namespace UE::RivermaxCore::Private
 {
+	/** Sidecar base struct given back to callback after data has been copied */
+	struct FBaseDataCopySideCar
+	{
+		virtual ~FBaseDataCopySideCar() = default;
+	};
+
+	/** Arguments describing memcopy operation required from allocator */
+	struct FCopyArgs
+	{
+		/** RHI source memory that will be mapped to cuda */
+		FBufferRHIRef RHISourceMemory;
+		
+		/** Source memory that doesn't need to be mapped since it's in the same space as destination*/
+		void* SourceMemory = nullptr;
+		
+		/** Destination memory where data has to be copied */
+		void* DestinationMemory = nullptr;
+		
+		/** Bytes needed to be copied */
+		uint32 SizeToCopy = 0;
+		
+		/** Sidecar to be re-provided on callback after memory has been copied */
+		TSharedPtr<FBaseDataCopySideCar> SideCar;
+	};
+
 	struct FRivermaxOutputFrame;
 
 	/** Delegate triggered when data has been copied into a given frame */
-	DECLARE_DELEGATE_OneParam(FOnFrameDataCopiedDelegate, const TSharedPtr<FRivermaxOutputFrame>&);
+	DECLARE_DELEGATE_OneParam(FOnFrameDataCopiedDelegate, const TSharedPtr<FBaseDataCopySideCar>&);
 
 	/** Base class for frame allocation. Currently supports GPU allocation using CUDA and System memory */
 	class FBaseFrameAllocator
 	{
 	public:
-		FBaseFrameAllocator(FIntPoint FrameResolution, uint32 Stride, FOnFrameDataCopiedDelegate InOnDataCopiedDelegate);
+		FBaseFrameAllocator(int32 InDesiredFrameSize, FOnFrameDataCopiedDelegate InOnDataCopiedDelegate);
 		virtual ~FBaseFrameAllocator() = default;
 
-		/** Allocates memory to hold FrameCount frames */
-		virtual bool Allocate(int32 FrameCount) = 0;
+		/** Allocates memory to hold FrameCount frames. 
+		 *  If bAlignFrameMemory is true, aligns each frame memory size to allocator alignment
+		 *  If not, aligns entire memory block instead
+		 */
+		virtual bool Allocate(int32 FrameCount, bool bAlignFrameMemory) = 0;
 
 		/** Deallocate memory */
 		virtual void Deallocate() = 0;
 
-		/** Initiate a copy into a given frame */
-		virtual bool CopyData(const FRivermaxOutputVideoFrameInfo& NewFrame, const TSharedPtr<FRivermaxOutputFrame>& DestinationFrame) = 0;
-		
-		/** Returns allocated frame for a given index */
-		const TSharedPtr<FRivermaxOutputFrame> GetFrame(int32 Index) const;
-
-	protected:
+		/** Initiate a copy into a given frame using Args description of the operation */
+		virtual bool CopyData(const FCopyArgs& Args) = 0;
 		
 		/** Returns buffer address of a given frame */
 		virtual void* GetFrameAddress(int32 FrameIndex) const = 0;
 
 	protected:
 
-		/** List of allocated frames. */
-		TArray<TSharedPtr<FRivermaxOutputFrame>> Frames;
-
-		/** Resolution of frames to be allocated */
-		FIntPoint FrameResolution = FIntPoint::ZeroValue;
-
-		/** Stride of a line */
-		uint32 Stride = 0;
-		
 		/** Total desired size of a frame */
 		int32 DesiredFrameSize = 0;
 
@@ -69,15 +83,15 @@ namespace UE::RivermaxCore::Private
 		using Super = FBaseFrameAllocator;
 
 	public:
-		FGPUAllocator(FIntPoint FrameResolution, uint32 Stride, FOnFrameDataCopiedDelegate InOnDataCopiedDelegate);
+		FGPUAllocator(int32 InDesiredFrameSize, FOnFrameDataCopiedDelegate InOnDataCopiedDelegate);
 
 		//~ Begin FBaseFrameAllocator interface
-		virtual bool Allocate(int32 FrameCount) override;
+		virtual bool Allocate(int32 FrameCount, bool bAlignFrameMemory) override;
 		virtual void Deallocate() override;
-		virtual bool CopyData(const FRivermaxOutputVideoFrameInfo& NewFrame, const TSharedPtr<FRivermaxOutputFrame>& DestinationFrame) override;
+		virtual bool CopyData(const FCopyArgs& Args) override;
+		virtual void* GetFrameAddress(int32 FrameIndex) const override;
 	
 	protected:
-		virtual void* GetFrameAddress(int32 FrameIndex) const override;
 		//~ End FBaseFrameAllocator interface
 
 	private:
@@ -104,8 +118,8 @@ namespace UE::RivermaxCore::Private
 		/** Map between buffer we are sending and their mapped address in gpu space */
 		TMap<FBufferRHIRef, void*> BufferCudaMemoryMap;
 
-		/** Frames that are waiting for a cuda memcopy. Will be dequeued when our cuda task has completed */
-		TSpscQueue<TSharedPtr<FRivermaxOutputFrame>> PendingCopies;
+		/** Sidecar for data copied callback that are waiting for a cuda memcopy. Will be dequeued when our cuda task has completed */
+		TSpscQueue<TSharedPtr<FBaseDataCopySideCar>> PendingCopies;
 	};
 
 	class FSystemAllocator : public FBaseFrameAllocator
@@ -113,14 +127,12 @@ namespace UE::RivermaxCore::Private
 		using Super = FBaseFrameAllocator;
 
 	public:
-		FSystemAllocator(FIntPoint FrameResolution, uint32 Stride, FOnFrameDataCopiedDelegate InOnDataCopiedDelegate);
+		FSystemAllocator(int32 InDesiredFrameSize, FOnFrameDataCopiedDelegate InOnDataCopiedDelegate);
 		
 		//~ Begin FBaseFrameAllocator interface
-		virtual bool Allocate(int32 FrameCount) override;
+		virtual bool Allocate(int32 FrameCount, bool bAlignFrameMemory) override;
 		virtual void Deallocate() override;
-		virtual bool CopyData(const FRivermaxOutputVideoFrameInfo& NewFrame, const TSharedPtr<FRivermaxOutputFrame>& DestinationFrame) override;
-	
-	protected:
+		virtual bool CopyData(const FCopyArgs& Args) override;
 		virtual void* GetFrameAddress(int32 FrameIndex) const override;
 		//~ End FBaseFrameAllocator interface
 	
@@ -137,6 +149,12 @@ namespace UE::RivermaxCore::Private
 
 		/** Allocated memory base address used when it's time to free */
 		void* SystemMemoryBaseAddress = nullptr;
+
+		/** Cached value for cvar controlling whether parallelfor is used to memcopy */
+		bool bUseParallelFor = false;
+
+		/** Cached value for cvar controlling number of threads to use when doing parallelfor memcopy */
+		int32 ParallelForThreadCount = 4;
 	};
 }
 
diff --git a/Engine/Plugins/VirtualProduction/Rivermax/RivermaxCore/Source/RivermaxCore/Private/RivermaxFrameManager.cpp b/Engine/Plugins/VirtualProduction/Rivermax/RivermaxCore/Source/RivermaxCore/Private/RivermaxFrameManager.cpp
--- a/Engine/Plugins/VirtualProduction/Rivermax/RivermaxCore/Source/RivermaxCore/Private/RivermaxFrameManager.cpp
+++ b/Engine/Plugins/VirtualProduction/Rivermax/RivermaxCore/Source/RivermaxCore/Private/RivermaxFrameManager.cpp
@@ -8,6 +8,12 @@
 
 namespace UE::RivermaxCore::Private
 {
+	/** Sidecar used when initiating memcopy. We provide the frame involved to update its state. */
+	struct FFrameBufferCopyInfo : public FBaseDataCopySideCar
+	{
+		TSharedPtr<FRivermaxOutputFrame> CopiedFrame;
+	};
+
 	FFrameManager::~FFrameManager()
 	{
 		Cleanup();
@@ -24,13 +30,12 @@ namespace UE::RivermaxCore::Private
 		OnCriticalErrorDelegate = Args.OnCriticalErrorDelegate;
 		FrameResolution = Args.Resolution;
 		TotalFrameCount = Args.NumberOfFrames;
-		const uint32 FrameSize = FrameResolution.Y * Args.Stride;
 		FOnFrameDataCopiedDelegate OnDataCopiedDelegate = FOnFrameDataCopiedDelegate::CreateRaw(this, &FFrameManager::OnDataCopied);
 
 		if (Args.bTryGPUAllocation)
 		{
-			FrameAllocator = MakeUnique<FGPUAllocator>(FrameResolution, Args.Stride, OnDataCopiedDelegate);
-			if (FrameAllocator->Allocate(TotalFrameCount))
+			FrameAllocator = MakeUnique<FGPUAllocator>(Args.FrameDesiredSize, OnDataCopiedDelegate);
+			if (FrameAllocator->Allocate(TotalFrameCount, Args.bAlignEachFrameAlloc))
 			{
 				MemoryLocation = EFrameMemoryLocation::GPU;
 			}
@@ -38,8 +43,8 @@ namespace UE::RivermaxCore::Private
 
 		if (MemoryLocation == EFrameMemoryLocation::None)
 		{
-			FrameAllocator = MakeUnique<FSystemAllocator>(FrameResolution, Args.Stride, OnDataCopiedDelegate);
-			if (FrameAllocator->Allocate(TotalFrameCount))
+			FrameAllocator = MakeUnique<FSystemAllocator>(Args.FrameDesiredSize, OnDataCopiedDelegate);
+			if (FrameAllocator->Allocate(TotalFrameCount, Args.bAlignEachFrameAlloc))
 			{
 				MemoryLocation = EFrameMemoryLocation::System;
 			}
@@ -55,6 +60,11 @@ namespace UE::RivermaxCore::Private
 			{
 				// All frames default to being available
 				FreeFrames.Add(Index);
+
+				/** Create actual frames and assign their video memory address from allocator */
+				TSharedPtr<FRivermaxOutputFrame> Frame = MakeShared<FRivermaxOutputFrame>(Index);
+				Frame->VideoBuffer = FrameAllocator->GetFrameAddress(Index);
+				Frames.Add(MoveTemp(Frame));
 			}
 		}
 
@@ -75,7 +85,7 @@ namespace UE::RivermaxCore::Private
 		FScopeLock Lock(&ContainersCritSec);
 		if (!FreeFrames.IsEmpty())
 		{
-			return FrameAllocator->GetFrame(FreeFrames[0]);
+			return Frames[FreeFrames[0]];
 		}
 
 		return nullptr;
@@ -119,7 +129,7 @@ namespace UE::RivermaxCore::Private
 		{
 			for (const uint32 FrameIndex : PendingFrames)
 			{
-				TSharedPtr<FRivermaxOutputFrame> Frame = FrameAllocator->GetFrame(FrameIndex);
+				TSharedPtr<FRivermaxOutputFrame> Frame = Frames[FrameIndex];
 				if (Frame->FrameIdentifier == FrameIdentifier)
 				{
 					// We found a reserved / pending frame corresponding to next frame identifier
@@ -136,7 +146,7 @@ namespace UE::RivermaxCore::Private
 		if (!ReadyFrames.IsEmpty())
 		{
 			const uint32 NextReadyFrame = ReadyFrames[0];
-			return FrameAllocator->GetFrame(NextReadyFrame);
+			return Frames[NextReadyFrame];
 		}
 	
 		return nullptr;
@@ -186,7 +196,7 @@ namespace UE::RivermaxCore::Private
 
 	const TSharedPtr<FRivermaxOutputFrame> FFrameManager::GetFrame(int32 Index) const
 	{
-		return FrameAllocator->GetFrame(Index);
+		return Frames[Index];
 	}
 
 	bool FFrameManager::SetFrameData(const FRivermaxOutputVideoFrameInfo& NewFrameInfo)
@@ -194,31 +204,49 @@ namespace UE::RivermaxCore::Private
 		bool bSuccess = false;
 		if (TSharedPtr<FRivermaxOutputFrame> NextFrame = GetNextFrame(NewFrameInfo.FrameIdentifier))
 		{
-			 bSuccess = FrameAllocator->CopyData(NewFrameInfo, NextFrame);
-			 if (!bSuccess)
-			 {
-				 OnCriticalErrorDelegate.ExecuteIfBound();
-			 }
+			TSharedPtr<FFrameBufferCopyInfo> Sidecar = MakeShared<FFrameBufferCopyInfo>();
+			Sidecar->CopiedFrame = NextFrame;
+
+			FCopyArgs Args;
+			Args.RHISourceMemory = NewFrameInfo.GPUBuffer;
+			Args.SourceMemory = NewFrameInfo.VideoBuffer;
+			Args.DestinationMemory = NextFrame->VideoBuffer;
+			Args.SizeToCopy = NewFrameInfo.Height * NewFrameInfo.Stride;
+			Args.SideCar = MoveTemp(Sidecar);
+
+			bSuccess = FrameAllocator->CopyData(Args);
+
+			if (!bSuccess)
+			{
+				OnCriticalErrorDelegate.ExecuteIfBound();
+			}
 			 
 		}
 		return bSuccess;
 	}
 
-	void FFrameManager::OnDataCopied(const TSharedPtr<FRivermaxOutputFrame>& Frame)
+	void FFrameManager::OnDataCopied(const TSharedPtr<FBaseDataCopySideCar>& Payload)
 	{
-		OnPreFrameReadyDelegate.ExecuteIfBound();
-
-		if (TSharedPtr<FRivermaxOutputFrame> AvailableFrame = GetNextFrame(Frame->FrameIdentifier))
+		TSharedPtr<FFrameBufferCopyInfo> CopyInfo = StaticCastSharedPtr<FFrameBufferCopyInfo>(Payload);
+		if (ensure(CopyInfo && CopyInfo->CopiedFrame))
 		{
-			if (AvailableFrame->IsReadyToBeSent())
+			OnPreFrameReadyDelegate.ExecuteIfBound();
+
+			if (TSharedPtr<FRivermaxOutputFrame> AvailableFrame = GetNextFrame(CopyInfo->CopiedFrame->FrameIdentifier))
 			{
-				const FString TraceName = FString::Format(TEXT("Rmax::FrameReady {0}"), { AvailableFrame->FrameIndex });
-				TRACE_CPUPROFILER_EVENT_SCOPE_TEXT(*TraceName);
-				TRACE_CPUPROFILER_EVENT_SCOPE_TEXT(*FString::Printf(TEXT("MediaCapturePipe: %u"), AvailableFrame->FrameIdentifier));
-				
-				MarkAsReady(AvailableFrame);
+				// Video frame has been copied, update frame's state
+				AvailableFrame->bIsVideoBufferReady = true;
+				if (AvailableFrame->IsReadyToBeSent())
+				{
+					const FString TraceName = FString::Format(TEXT("Rmax::FrameReady {0}"), { AvailableFrame->FrameIndex });
+					TRACE_CPUPROFILER_EVENT_SCOPE_TEXT(*TraceName);
+					TRACE_CPUPROFILER_EVENT_SCOPE_TEXT(*FString::Printf(TEXT("MediaCapturePipe: %u"), AvailableFrame->FrameIdentifier));
+
+					MarkAsReady(AvailableFrame);
+				}
 			}
 		}
+		
 	}
 }
 
diff --git a/Engine/Plugins/VirtualProduction/Rivermax/RivermaxCore/Source/RivermaxCore/Private/RivermaxFrameManager.h b/Engine/Plugins/VirtualProduction/Rivermax/RivermaxCore/Source/RivermaxCore/Private/RivermaxFrameManager.h
--- a/Engine/Plugins/VirtualProduction/Rivermax/RivermaxCore/Source/RivermaxCore/Private/RivermaxFrameManager.h
+++ b/Engine/Plugins/VirtualProduction/Rivermax/RivermaxCore/Source/RivermaxCore/Private/RivermaxFrameManager.h
@@ -18,6 +18,7 @@ namespace UE::RivermaxCore
 namespace UE::RivermaxCore::Private
 {
 	class FBaseFrameAllocator;
+	struct FBaseDataCopySideCar;
 
 	/** Where frame memory is allocated */
 	enum class EFrameMemoryLocation : uint8
@@ -46,6 +47,12 @@ namespace UE::RivermaxCore::Private
 		/** Stride of a line of video frame */
 		uint32 Stride = 0;
 
+		/** Desired size for a frame. Can be greater than what is needed to align with Rivermax's chunks */
+		uint32 FrameDesiredSize = 0;
+
+		/** Whether allocator will align each frame to desired alignment or the entire block */
+		bool bAlignEachFrameAlloc = false;
+
 		/** Number of video frames required */
 		uint8 NumberOfFrames = 0;
 
@@ -134,7 +141,7 @@ namespace UE::RivermaxCore::Private
 	private:
 
 		/** Called back when copy request was completed by allocator */
-		void OnDataCopied(const TSharedPtr<FRivermaxOutputFrame>& Frame);
+		void OnDataCopied(const TSharedPtr<FBaseDataCopySideCar>& Sidecar);
 
 	private:
 
@@ -153,6 +160,9 @@ namespace UE::RivermaxCore::Private
 		/** Frame allocator dealing with memory operation */
 		TUniquePtr<FBaseFrameAllocator> FrameAllocator;
 
+		/** List of allocated frames. */
+		TArray<TSharedPtr<FRivermaxOutputFrame>> Frames;
+
 		/** Delegate triggered when a frame is free to use */
 		FOnFreeFrameDelegate OnFreeFrameDelegate;
 
diff --git a/Engine/Plugins/VirtualProduction/Rivermax/RivermaxCore/Source/RivermaxCore/Private/RivermaxOutputFrame.h b/Engine/Plugins/VirtualProduction/Rivermax/RivermaxCore/Source/RivermaxCore/Private/RivermaxOutputFrame.h
--- a/Engine/Plugins/VirtualProduction/Rivermax/RivermaxCore/Source/RivermaxCore/Private/RivermaxOutputFrame.h
+++ b/Engine/Plugins/VirtualProduction/Rivermax/RivermaxCore/Source/RivermaxCore/Private/RivermaxOutputFrame.h
@@ -56,13 +56,23 @@ namespace UE::RivermaxCore::Private
 		/** Timestamp of this frame  */
 		double TimestampTicks = 0.0;
 
-		/** Pointers retrieved from Rivermax for the next chunk */
+		/** Payload (data) pointer retrieved from Rivermax for the next chunk */
 		void* PayloadPtr = nullptr;
+		
+		/** Header pointer retrieved from Rivermax for the next chunk */
 		void* HeaderPtr = nullptr;
+		
+		/** Cached address of beginning of frame in Rivermax's memblock. Used when using intermediate buffer/ */
+		void* FrameStartPtr = nullptr;
+		
+		/** Offset in the frame where we are at to copy next block of data */
+		uint32 Offset = 0;
+		
+		/** Time covered by memcopies. Not used for now but it could be used to detect if memory will be used before it's actually copied. */
+		uint64 ScheduleTimeCopied = 0;
 
 		/** Time at which this frame was made available to be sent */
 		uint64 ReadyTimestamp = 0;
-
 	};
 }
 
diff --git a/Engine/Plugins/VirtualProduction/Rivermax/RivermaxCore/Source/RivermaxCore/Private/RivermaxOutputStream.cpp b/Engine/Plugins/VirtualProduction/Rivermax/RivermaxCore/Source/RivermaxCore/Private/RivermaxOutputStream.cpp
--- a/Engine/Plugins/VirtualProduction/Rivermax/RivermaxCore/Source/RivermaxCore/Private/RivermaxOutputStream.cpp
+++ b/Engine/Plugins/VirtualProduction/Rivermax/RivermaxCore/Source/RivermaxCore/Private/RivermaxOutputStream.cpp
@@ -8,6 +8,7 @@
 #include "IRivermaxBoundaryMonitor.h"
 #include "IRivermaxManager.h"
 #include "Misc/ByteSwap.h"
+#include "RivermaxFrameAllocator.h"
 #include "RivermaxFrameManager.h"
 #include "RivermaxLog.h"
 #include "RivermaxPTPUtils.h"
@@ -57,6 +58,31 @@ namespace UE::RivermaxCore::Private
 		TEXT("If not 0, overrides transmit offset calculation (TRO) based on  frame rate and resolution with a fixed value. Value in seconds."),
 		ECVF_Default);
 
+	static TAutoConsoleVariable<int32> CVarRivermaxOutputSkipSchedulingCutOffTime(
+		TEXT("Rivermax.Output.Scheduling.SkipCutoff"), 100,
+		TEXT("Required time in microseconds from scheduling time to avoid skipping an interval."),
+		ECVF_Default);
+
+	static TAutoConsoleVariable<int32> CVarRivermaxOutputForceImmediateSchedulingThreshold(
+		TEXT("Rivermax.Output.Scheduling.ForceImmediateCutoff"), 600,
+		TEXT("Required time in nanoseconds from scheduling time before we clamp to do it immediately."),
+		ECVF_Default);
+
+	static TAutoConsoleVariable<int32> CVarRivermaxOutputFrameSliceCount(
+		TEXT("Rivermax.Output.FrameSliceCount"), 20,
+		TEXT("Number of memcopy done per frame when using intermediate buffer. As frame gets bigger, we can't do a single memcopy or timings will be broken."),
+		ECVF_Default);
+
+	static TAutoConsoleVariable<int32> CVarRivermaxOutputEnableIntermediateBuffer(
+		TEXT("Rivermax.Output.EnableIntermediateBuffer"), 1,
+		TEXT("Uses an intermediate buffer used by Rivermax when sending data out. During scheduling, captured frame data will be copied over intermediate buffer."),
+		ECVF_Default);
+
+	static TAutoConsoleVariable<int32> CVarRivermaxOutputUseSingleMemblock(
+		TEXT("Rivermax.Output.UseSingleMemblock"), 1,
+		TEXT("Configures Rivermax stream to use a single memblock potentially improving SDK performance."),
+		ECVF_Default);
+
 	bool FindPayloadSize(const FRivermaxOutputStreamOptions& InOptions, uint32 InBytesPerLine, const FVideoFormatInfo& FormatInfo, uint16& OutPayloadSize)
 	{
 		using namespace UE::RivermaxCore::Private::Utils;
@@ -294,6 +320,13 @@ namespace UE::RivermaxCore::Private
 
 	}
 
+	/** Sidecar used when copying  parts of frame data while scheduling */
+	struct FMemChunkCopiedInfo : public FBaseDataCopySideCar
+	{
+		/** Time covered inside frame interval including last chunk copied. */
+		uint64 TimeCovered = 0;
+	};
+
 	FRivermaxOutputStream::~FRivermaxOutputStream()
 	{
 		Uninitialize();
@@ -314,6 +347,15 @@ namespace UE::RivermaxCore::Private
 		Listener = &InListener;
 		FormatInfo = FStandardVideoFormat::GetVideoFormatInfo(Options.PixelFormat);
 
+		const uint32 FrameSize = GetStride() * Options.AlignedResolution.Y;
+		if (FrameSize == 0)
+		{
+			UE_LOG(LogRivermax, Warning, TEXT("Can't create Rivermax Output Stream. Frame size of 0 is invalid. Verify resolution."));
+			return false;
+		}
+		
+		CacheCVarValues();
+
 		// Verify resolution for sampling type
 		if (Options.AlignedResolution.X % FormatInfo.PixelGroupCoverage != 0)
 		{
@@ -321,7 +363,7 @@ namespace UE::RivermaxCore::Private
 			return false;
 		}
 
-		if (!SetupFrameManagement())
+		if (!InitializeStreamMemoryConfig())
 		{
 			return false;
 		}
@@ -334,8 +376,7 @@ namespace UE::RivermaxCore::Private
 			TAnsiStringBuilder<2048> SDPDescription;
 			UE::RivermaxCore::Private::Utils::StreamOptionsToSDPDescription(Options, SDPDescription);
 
-			// Initialize internal memory and rivermax configuration 
-			if(InitializeStreamMemoryConfig())
+			// Create Rivermax stream using memory configuration
 			{
 				// Initialize rivermax output stream with desired config
 				const uint32 NumberPacketsPerFrame = StreamMemory.PacketsPerFrame;
@@ -426,7 +467,14 @@ namespace UE::RivermaxCore::Private
 	{
 		TRACE_CPUPROFILER_EVENT_SCOPE(FRivermaxOutputStream::PushVideoFrame);
 
-		return FrameManager->SetFrameData(NewFrame);
+		FRivermaxOutputVideoFrameInfo CopyInfo;
+		CopyInfo.Stride = GetStride();
+		CopyInfo.Height = Options.AlignedResolution.Y;
+		CopyInfo.FrameIdentifier = NewFrame.FrameIdentifier;
+		CopyInfo.GPUBuffer = NewFrame.GPUBuffer;
+		CopyInfo.VideoBuffer = NewFrame.VideoBuffer;
+		CopyInfo.Width = NewFrame.Width;
+		return FrameManager->SetFrameData(CopyInfo);
 	}
 
 	void FRivermaxOutputStream::Process_AnyThread()
@@ -446,60 +494,74 @@ namespace UE::RivermaxCore::Private
 				WaitForNextRound();
 			}
 
+			// Verify if we can continue sending the next frame. We might be too late if our thread wasn't woken up in time
+			if (CanProceedWithScheduling())
 			{
-				TRACE_CPUPROFILER_EVENT_SCOPE(FRivermaxOutputStream::PrepareNextFrame);
-				PrepareNextFrame();
-			}
 
-			// At this point, if there is no frame to send, move on to wait for next round
-			if(CurrentFrame.IsValid() == false || !bIsActive)
-			{
-				return;
-			}
+				Stats.LastFrameChunkRetries = 0;
+				{
+					TRACE_CPUPROFILER_EVENT_SCOPE(FRivermaxOutputStream::PrepareNextFrame);
+					PrepareNextFrame();
+				}
 
-			StreamData.LastSendStartTimeNanoSec = RivermaxModule->GetRivermaxManager()->GetTime();
-			
-			// Add markup when we start pushing out a frame with its timestamp to track it across network
-			if (StreamData.bHasFrameFirstChunkBeenFetched == false)
-			{
-				const uint32 MediaTimestamp = GetTimestampFromTime(StreamData.NextAlignmentPointNanosec, MediaClockSampleRate);
-				const FString TraceName = FString::Format(TEXT("RmaxOutput::StartSending {0}|{1}"), { MediaTimestamp, CurrentFrame->FrameIdentifier });
-				TRACE_CPUPROFILER_EVENT_SCOPE_TEXT(*TraceName);
-				UE_LOG(LogRivermax, Verbose, TEXT("RmaxTX frame number %u with timestamp %u."), CurrentFrame->FrameIdentifier, MediaTimestamp);
-			}
+				// At this point, if there is no frame to send, move on to wait for next round
+				if (CurrentFrame.IsValid() && bIsActive)
+				{
+					StreamData.LastSendStartTimeNanoSec = RivermaxModule->GetRivermaxManager()->GetTime();
 
-			const FString TraceName = FString::Format(TEXT("RmaxOutput::SendFrame {0}"), { CurrentFrame->FrameIndex });
-			TRACE_CPUPROFILER_EVENT_SCOPE_TEXT(*TraceName);
-			TRACE_CPUPROFILER_EVENT_SCOPE_TEXT(*FString::Printf(TEXT("MediaCapturePipe: %u"), CurrentFrame->FrameIdentifier));
+					// Add markup when we start pushing out a frame with its timestamp to track it across network
+					if (StreamData.bHasFrameFirstChunkBeenFetched == false)
+					{
+						const uint32 MediaTimestamp = GetTimestampFromTime(StreamData.NextScheduleTimeNanosec, MediaClockSampleRate);
+						const FString TraceName = FString::Format(TEXT("RmaxOutput::StartSending {0}|{1}"), { MediaTimestamp, CurrentFrame->FrameIdentifier });
+						TRACE_CPUPROFILER_EVENT_SCOPE_TEXT(*TraceName);
+						UE_LOG(LogRivermax, Verbose, TEXT("RmaxTX frame number %u with timestamp %u."), CurrentFrame->FrameIdentifier, MediaTimestamp);
+					}
 
-			do 
-			{
-				if (bIsActive)
-				{	
-					GetNextChunk();
-				}
+					const FString TraceName = FString::Format(TEXT("RmaxOutput::SendFrame {0}"), { CurrentFrame->FrameIndex });
+					TRACE_CPUPROFILER_EVENT_SCOPE_TEXT(*TraceName);
+					TRACE_CPUPROFILER_EVENT_SCOPE_TEXT(*FString::Printf(TEXT("MediaCapturePipe: %u"), CurrentFrame->FrameIdentifier));
 
-				if (bIsActive)
-				{
-					SetupRTPHeaders();
-				}
+					do 
+					{
+						if (bIsActive)
+						{
+							GetNextChunk();
+						}
 
-				if (bIsActive)
-				{
-					CommitNextChunks();
-				}
+						if (bIsActive && CachedCVars.bUseIntermediateBuffer)
+						{
+							CopyFrameData();
+						}
 
-				//Update frame progress
-				if (bIsActive)
-				{
-					Stats.TotalStrides += StreamMemory.PacketsPerChunk;
-					++CurrentFrame->ChunkNumber;
-				}
+						if (bIsActive)
+						{
+							SetupRTPHeaders();
+						}
 
-			} while (CurrentFrame->ChunkNumber < StreamMemory.ChunksPerMemoryBlock && bIsActive);
-			
+						if (bIsActive)
+						{
+							CommitNextChunks();
+						}
+
+						//Update frame progress
+						if (bIsActive)
+						{
+							Stats.TotalStrides += StreamMemory.PacketsPerChunk;
+							++CurrentFrame->ChunkNumber;
+						}
+					} while (CurrentFrame->ChunkNumber < StreamMemory.ChunksPerFrameField && bIsActive);
+					
+					Stats.FramesSentCounter++;
+					StreamData.RTPHeaderFrameIndex = (StreamData.RTPHeaderFrameIndex + 1) % Options.NumberOfBuffers; 
+				}
+			}
+			else
+			{
+				Stats.LastFrameChunkRetries = 0;
+			}
 
-			Stats.MemoryBlockSentCounter++;
+			Stats.TotalChunkRetries += Stats.LastFrameChunkRetries;
 			StreamData.bHasFrameFirstChunkBeenFetched = false;
 		}
 	}
@@ -552,15 +614,16 @@ namespace UE::RivermaxCore::Private
 
 		StreamMemory.LinesInChunk = FindLinesPerChunk(Options);
 		StreamMemory.PacketsPerChunk = StreamMemory.LinesInChunk * StreamMemory.PacketsInLine;
-		StreamMemory.FramesFieldPerMemoryBlock = 1;
+		StreamMemory.FramesFieldPerMemoryBlock = CachedCVars.bUseSingleMemblock ? Options.NumberOfBuffers : 1;
 
 		// Chunk count won't necessarily align with the number of packets required. We need an integer amount of chunks to initialize our stream
 		// and calculate how many packets that represents. Rivermax will expect the payload/header array to be that size. It just means that
 		// we will mark the extra packets as 0 size.
 		StreamMemory.ChunksPerFrameField = FMath::CeilToInt32((float)StreamMemory.PacketsPerFrame / StreamMemory.PacketsPerChunk);
-		StreamMemory.PacketsPerMemoryBlock = StreamMemory.ChunksPerFrameField * StreamMemory.PacketsPerChunk * StreamMemory.FramesFieldPerMemoryBlock;
+		const uint64 RealPacketsPerFrame = StreamMemory.ChunksPerFrameField * StreamMemory.PacketsPerChunk;
+		StreamMemory.PacketsPerMemoryBlock = RealPacketsPerFrame * StreamMemory.FramesFieldPerMemoryBlock;
 		StreamMemory.ChunksPerMemoryBlock = StreamMemory.FramesFieldPerMemoryBlock * StreamMemory.ChunksPerFrameField;
-		StreamMemory.MemoryBlockCount = Options.NumberOfBuffers;
+		StreamMemory.MemoryBlockCount = Options.NumberOfBuffers / StreamMemory.FramesFieldPerMemoryBlock;
 
 		// Setup arrays with the right sizes so we can give pointers to rivermax
 		StreamMemory.RTPHeaders.SetNumZeroed(StreamMemory.MemoryBlockCount);
@@ -568,50 +631,60 @@ namespace UE::RivermaxCore::Private
 		StreamMemory.HeaderSizes.SetNumUninitialized(StreamMemory.PacketsPerMemoryBlock);
 		StreamMemory.HeaderStrideSize = sizeof(FRawRTPHeader);
 
-		uint64 TotalSize = 0;
-		uint64 LineSize = 0;
-		for (int32 PayloadSizeIndex = 0; PayloadSizeIndex < StreamMemory.PayloadSizes.Num(); ++PayloadSizeIndex)
+	
+		for (uint32 BufferIndex = 0; BufferIndex < StreamMemory.FramesFieldPerMemoryBlock; ++BufferIndex)
 		{
-			uint32 HeaderSize = FRawRTPHeader::OneSRDSize;
-			uint32 ThisPayloadSize = StreamMemory.PayloadSize;
-			if (TotalSize < FrameSize)
+			uint64 TotalSize = 0;
+			uint64 LineSize = 0;
+			for (int32 PayloadSizeIndex = 0; PayloadSizeIndex < RealPacketsPerFrame; ++PayloadSizeIndex)
 			{
-				if ((LineSize + StreamMemory.PayloadSize) == BytesPerLine)
-				{
-					LineSize = 0;
-				}
-				else if ((LineSize + StreamMemory.PayloadSize) > BytesPerLine)
+				uint32 HeaderSize = FRawRTPHeader::OneSRDSize;
+				uint32 ThisPayloadSize = StreamMemory.PayloadSize;
+				if (TotalSize < FrameSize)
 				{
-					HeaderSize = FRawRTPHeader::TwoSRDSize;
-					LineSize = StreamMemory.PayloadSize - (BytesPerLine - LineSize);
-					if (LineSize > BytesPerLine)
+					if ((LineSize + StreamMemory.PayloadSize) == BytesPerLine)
 					{
-						UE_LOG(LogRivermax, Warning, TEXT("Unsupported small resolution, %dx%d, needing more than 2 SRD to express"), Options.AlignedResolution.X, Options.AlignedResolution.Y);
-						return false;
+						LineSize = 0;
+					}
+					else if ((LineSize + StreamMemory.PayloadSize) > BytesPerLine)
+					{
+						HeaderSize = FRawRTPHeader::TwoSRDSize;
+						LineSize = StreamMemory.PayloadSize - (BytesPerLine - LineSize);
+						if (LineSize > BytesPerLine)
+						{
+							UE_LOG(LogRivermax, Warning, TEXT("Unsupported small resolution, %dx%d, needing more than 2 SRD to express"), Options.AlignedResolution.X, Options.AlignedResolution.Y);
+							return false;
+						}
+					}
+					else
+					{
+						// Keep track of line size offset to know when to use TwoSRDs
+						LineSize += StreamMemory.PayloadSize;
+					}
+
+					if ((TotalSize + StreamMemory.PayloadSize) > FrameSize)
+					{
+						HeaderSize = FRawRTPHeader::OneSRDSize;
 					}
 				}
 				else
 				{
-					// Keep track of line size offset to know when to use TwoSRDs
-					LineSize += StreamMemory.PayloadSize;
+					// Extra header/payload required for the chunk alignment are set to 0. Nothing has to be sent out the wire.
+					HeaderSize = 0;
+					ThisPayloadSize = 0;
 				}
 
-				if ((TotalSize + StreamMemory.PayloadSize) > FrameSize)
-				{
-					HeaderSize = FRawRTPHeader::OneSRDSize;
-				}
-			}
-			else
-			{
-				// Extra header/payload required for the chunk alignment are set to 0. Nothing has to be sent out the wire.
-				HeaderSize = 0;
-				ThisPayloadSize = 0;
+				StreamMemory.HeaderSizes[PayloadSizeIndex + (BufferIndex* RealPacketsPerFrame)] = HeaderSize;
+				StreamMemory.PayloadSizes[PayloadSizeIndex + (BufferIndex * RealPacketsPerFrame)] = ThisPayloadSize;
+				TotalSize += ThisPayloadSize;
 			}
-			
-			StreamMemory.HeaderSizes[PayloadSizeIndex] = HeaderSize;
-			StreamMemory.PayloadSizes[PayloadSizeIndex] = ThisPayloadSize;
-			TotalSize += ThisPayloadSize;
 		}
+
+		if (!SetupFrameManagement())
+		{
+			return false;
+		}
+		
 		StreamMemory.MemoryBlocks.SetNumZeroed(StreamMemory.MemoryBlockCount);
 		for (uint32 BlockIndex = 0; BlockIndex < StreamMemory.MemoryBlockCount; ++BlockIndex)
 		{
@@ -619,9 +692,18 @@ namespace UE::RivermaxCore::Private
 			Block.chunks_num = StreamMemory.ChunksPerMemoryBlock;
 			Block.app_hdr_size_arr = StreamMemory.HeaderSizes.GetData();
 			Block.data_size_arr = StreamMemory.PayloadSizes.GetData();
-			Block.data_ptr = FrameManager->GetFrame(BlockIndex)->VideoBuffer;
-			
-			StreamMemory.RTPHeaders[BlockIndex].SetNumZeroed(StreamMemory.PacketsPerFrame);
+
+			// If intermediate buffer is used, we setup rmax memblock to use that address. Otherwise, we map it to our actual frame's address
+			if(CachedCVars.bUseIntermediateBuffer)
+			{
+				Block.data_ptr = Allocator->GetFrameAddress(BlockIndex);
+			}
+			else
+			{
+				Block.data_ptr = FrameManager->GetFrame(BlockIndex)->VideoBuffer;
+			}
+
+			StreamMemory.RTPHeaders[BlockIndex].SetNumZeroed(RealPacketsPerFrame * StreamMemory.FramesFieldPerMemoryBlock);
 			Block.app_hdr_ptr = &StreamMemory.RTPHeaders[BlockIndex][0];
 		}
 
@@ -636,6 +718,8 @@ namespace UE::RivermaxCore::Private
 		NextFrame->ChunkNumber = 0;
 		NextFrame->PayloadPtr = nullptr;
 		NextFrame->HeaderPtr = nullptr;
+		NextFrame->FrameStartPtr = nullptr;
+		NextFrame->ScheduleTimeCopied = 0;
 	}
 
 	void FRivermaxOutputStream::BuildRTPHeader(FRawRTPHeader& OutHeader) const
@@ -762,7 +846,7 @@ namespace UE::RivermaxCore::Private
 
 		// Offset wakeup if desired to give more time for scheduling. 
 		const uint64 WakeupTime = StreamData.NextAlignmentPointNanosec - CVarRivermaxWakeupOffset.GetValueOnAnyThread();
-		
+
 		uint64 WaitTimeNanosec = WakeupTime - CurrentTimeNanosec;
 
 		// Wakeup can be smaller than current time with controllable offset
@@ -774,6 +858,7 @@ namespace UE::RivermaxCore::Private
 		static constexpr float SleepThresholdSec = 5.0f * (1.0f / 1000.0f);
 		static constexpr float YieldTimeSec = 2.0f * (1.0f / 1000.0f);
 		const double WaitTimeSec = FMath::Min(WaitTimeNanosec / 1E9, 1.0);
+		StreamData.LastSleepTimeNanoSec = WaitTimeNanosec;
 
 		// Sleep for the largest chunk of time
 		if (WaitTimeSec > SleepThresholdSec)
@@ -803,7 +888,17 @@ namespace UE::RivermaxCore::Private
 		if (StreamData.bHasValidNextFrameNumber)
 		{
 			const uint64 AfterSleepTimeNanosec = RivermaxModule->GetRivermaxManager()->GetTime();
-			UE_LOG(LogRivermax, Verbose, TEXT("Scheduling at %llu. CurrentTime %llu. NextAlign %llu. Waiting %0.9f"), StreamData.NextScheduleTimeNanosec, CurrentTimeNanosec, StreamData.NextAlignmentPointNanosec, (double)WaitTimeNanosec / 1E9);
+			const uint64 RealWaitNs = AfterSleepTimeNanosec - CurrentTimeNanosec;
+			const uint64 OvershootSleep = AfterSleepTimeNanosec > StreamData.NextAlignmentPointNanosec ? AfterSleepTimeNanosec - StreamData.NextAlignmentPointNanosec : 0;
+			const double OvershootSleepSec = OvershootSleep / 1e9;
+
+			UE_LOG(LogRivermax, Verbose, TEXT("CurrentTime %llu. OvershootSleep: %0.9f. ExpectedWait: %0.9f. RealWait: %0.9f, Scheduling at %llu. NextAlign %llu. ")
+			, CurrentTimeNanosec
+			, OvershootSleepSec
+			, (double)WaitTimeNanosec / 1E9
+			, (double)RealWaitNs / 1E9
+			, StreamData.NextScheduleTimeNanosec
+			, StreamData.NextAlignmentPointNanosec);
 		}
 	}
 
@@ -819,10 +914,18 @@ namespace UE::RivermaxCore::Private
 				if (StreamData.bHasFrameFirstChunkBeenFetched == false)
 				{
 					StreamData.bHasFrameFirstChunkBeenFetched = true;
-					if (CurrentFrame->VideoBuffer != CurrentFrame->PayloadPtr)
+
+					// Stamp frame start in order to copy frame data sequentially as we query chunks
+					CurrentFrame->FrameStartPtr = CurrentFrame->PayloadPtr;
+					CurrentFrame->Offset = 0;
+
+					// If we aren't using intermediate buffer, we expect video buffer to point to the next chunk's payload pointer
+					if (CachedCVars.bUseIntermediateBuffer == false)
 					{
-						//Debug code to track rivermax chunk processing
-						UE_LOG(LogRivermax, Warning, TEXT("Frame being sent (%d) doesn't match chunks being processed."), CurrentFrame->FrameIndex);
+						if (CurrentFrame->VideoBuffer != CurrentFrame->PayloadPtr)
+						{
+							UE_LOG(LogRivermax, Warning, TEXT("Frame being sent (%d) doesn't match chunks being processed."), CurrentFrame->FrameIndex);
+						}
 					}
 				}
 
@@ -831,9 +934,9 @@ namespace UE::RivermaxCore::Private
 			else if (Status == RMAX_ERR_NO_FREE_CHUNK)
 			{
 				//We should not be here
-				Stats.ChunkRetries++;
 				if (!bHasAddedTrace)
 				{
+					Stats.LastFrameChunkRetries++;
 					UE_LOG(LogRivermax, Verbose, TEXT("No free chunks to get for chunk '%u'. Waiting"), CurrentFrame->ChunkNumber);
 					TRACE_CPUPROFILER_EVENT_SCOPE(GetNextChunk::NoFreeChunk);
 					bHasAddedTrace = true;
@@ -848,17 +951,73 @@ namespace UE::RivermaxCore::Private
 		} while (Status != RMAX_OK && bIsActive);
 	}
 
+	void FRivermaxOutputStream::CopyFrameData()
+	{
+		// Make sure copy size doesn't go over frame size
+		const uint32 FrameSize = GetStride() * Options.AlignedResolution.Y;
+		if (ensure(FrameSize > 0))
+		{
+			const uint32 BlockSize = 1 + ((FrameSize - 1) / CachedCVars.FrameMemorySliceCount);
+			const uint32 MaxSize = FrameSize - CurrentFrame->Offset;
+			const uint32 CopySize = FMath::Min(BlockSize, MaxSize);
+
+			// Copy data until we have covered the whole frame. Last block might be smaller.
+			if (CopySize > 0)
+			{
+				TRACE_CPUPROFILER_EVENT_SCOPE(CopyFrameData);
+				const uint32 PacketsToCopy = CopySize / StreamMemory.PayloadSize;
+				const uint32 PacketOffset = CurrentFrame->ChunkNumber * StreamMemory.PacketsPerChunk;
+				const uint64 PacketDeltaTime = (StreamData.FrameFieldTimeIntervalNs - TransmitOffsetNanosec) / StreamMemory.PacketsPerFrame;
+
+				// Time covered in frame interval including this copy
+				const uint64 TimeCopied = StreamData.NextScheduleTimeNanosec + ((PacketOffset + PacketsToCopy) * PacketDeltaTime);
+
+				// Sidecar used to track copy progress
+				TSharedPtr<FMemChunkCopiedInfo> CopyInfo = MakeShared<FMemChunkCopiedInfo>();
+				CopyInfo->TimeCovered = TimeCopied;
+
+				FCopyArgs Args;
+				uint8* SourceStart = reinterpret_cast<uint8*>(CurrentFrame->VideoBuffer);
+				uint8* DestinationStart = reinterpret_cast<uint8*>(CurrentFrame->FrameStartPtr);
+				Args.SourceMemory = SourceStart + CurrentFrame->Offset;
+				Args.DestinationMemory = DestinationStart + CurrentFrame->Offset;
+				Args.SizeToCopy = CopySize;
+				Args.SideCar = MoveTemp(CopyInfo);
+
+				// Update memory offset for next copy
+				CurrentFrame->Offset += CopySize;
+
+				Allocator->CopyData(Args);
+			}
+		}
+		else
+		{
+			UE_LOG(LogRivermax, Error, TEXT("Invalid frame size detected while stream was active. Shutting down."));
+			Listener->OnStreamError();
+			Stop();
+		}
+	}
+
 	void FRivermaxOutputStream::SetupRTPHeaders()
 	{
 		FRawRTPHeader* HeaderRawPtr = reinterpret_cast<FRawRTPHeader*>(CurrentFrame->HeaderPtr);
 		check(HeaderRawPtr);
-		for (uint32 StrideIndex = 0; StrideIndex < StreamMemory.PacketsPerChunk && CurrentFrame->PacketCounter < StreamMemory.PacketsPerFrame; ++StrideIndex)
+
+		const uint32 MemblockIndex = StreamData.RTPHeaderFrameIndex / StreamMemory.FramesFieldPerMemoryBlock;
+		const uint32 MemblockFrameIndex = StreamData.RTPHeaderFrameIndex % StreamMemory.FramesFieldPerMemoryBlock;
+		const uint32 PacketsPerFrame = StreamMemory.ChunksPerFrameField * StreamMemory.PacketsPerChunk;
+		const uint32 BaseHeaderIndex = CurrentFrame->PacketCounter + (MemblockFrameIndex * PacketsPerFrame);
+
+		if (ensure(HeaderRawPtr == &StreamMemory.RTPHeaders[MemblockIndex][BaseHeaderIndex]))
 		{
-			BuildRTPHeader(*HeaderRawPtr);
-			++StreamData.SequenceNumber;
-			CurrentFrame->BytesSent += StreamMemory.PayloadSizes[CurrentFrame->PacketCounter];
-			++CurrentFrame->PacketCounter;
-			++HeaderRawPtr;
+			for (uint32 PacketIndex = 0; PacketIndex < StreamMemory.PacketsPerChunk && CurrentFrame->PacketCounter < StreamMemory.PacketsPerFrame; ++PacketIndex)
+			{
+				BuildRTPHeader(StreamMemory.RTPHeaders[MemblockIndex][BaseHeaderIndex + PacketIndex]);
+				
+				++StreamData.SequenceNumber;
+				++CurrentFrame->PacketCounter;
+				CurrentFrame->BytesSent += StreamMemory.PayloadSizes[CurrentFrame->PacketCounter];
+			}
 		}
 	}
 
@@ -867,6 +1026,7 @@ namespace UE::RivermaxCore::Private
 		rmax_status_t Status;
 		int32 ErrorCount = 0;
 		const uint64 CurrentTimeNanosec = RivermaxModule->GetRivermaxManager()->GetTime();
+
 		do
 		{
 			//Only first chunk gets scheduled with a timestamp. Following chunks are queued after it using 0
@@ -874,8 +1034,8 @@ namespace UE::RivermaxCore::Private
 			const rmax_commit_flags_t CommitFlags{};
 			if (ScheduleTime != 0)
 			{
-				//TRACE_BOOKMARK(TEXT("Sched A: %llu, C: %llu"), StreamData.NextAlignmentPointNanosec, CurrentTimeNanosec);
-				if (ScheduleTime <= CurrentTimeNanosec)
+				// If scheduling time is not far away enough, force it immediately otherwise rmax_commit will throw an error 
+				if (ScheduleTime <= (CurrentTimeNanosec + CachedCVars.ForceCommitImmediateTimeNanosec))
 				{
 					ScheduleTime = 0;
 					++Stats.CommitImmediate;
@@ -1041,36 +1201,44 @@ namespace UE::RivermaxCore::Private
 			// We will resend the last one so just reinitialize it to resend
 			InitializeNextFrame(CurrentFrame);
 			
-			// No frame to send, keep last one and restart its internal counters
-			UE_LOG(LogRivermax, Verbose, TEXT("No frame to send. Reusing last frame '%d' with Id %u"), CurrentFrame->FrameIndex, CurrentFrame->FrameIdentifier);
-
-			// Since we want to resend last frame, we need to fast forward chunk pointer to re-point to the one we just sent
-			rmax_status_t Status;
-			bool bHasAddedTrace = false;
-			do
+			// If intermediate buffer isn't used and frame has to be repeated, we use skip chunk method 
+			// which might cause timing errors caused by chunk managements issue 
+			if (!CachedCVars.bUseIntermediateBuffer)
 			{
-				Status = rmax_out_skip_chunks(StreamId, StreamMemory.ChunksPerFrameField * (Options.NumberOfBuffers - 1));
-				if (Status != RMAX_OK)
+				// No frame to send, keep last one and restart its internal counters
+				UE_LOG(LogRivermax, Verbose, TEXT("No frame to send. Reusing last frame '%d' with Id %u"), CurrentFrame->FrameIndex, CurrentFrame->FrameIdentifier);
+
+				// Since we want to resend last frame, we need to fast forward chunk pointer to re-point to the one we just sent
+				rmax_status_t Status;
+				bool bHasAddedTrace = false;
+				do
 				{
-					if (Status == RMAX_ERR_NO_FREE_CHUNK)
+					Status = rmax_out_skip_chunks(StreamId, StreamMemory.ChunksPerFrameField * (Options.NumberOfBuffers - 1));
+					if (Status != RMAX_OK)
 					{
-						// Wait until there are enough free chunk to be skipped
-						if (!bHasAddedTrace)
+						if (Status == RMAX_ERR_NO_FREE_CHUNK)
 						{
-							UE_LOG(LogRivermax, Warning, TEXT("No chunks ready to skip. Waiting"));
-							TRACE_CPUPROFILER_EVENT_SCOPE(PrepareNextFrame::NoFreeChunk);
-							bHasAddedTrace = true;
+							// Wait until there are enough free chunk to be skipped
+							if (!bHasAddedTrace)
+							{
+								UE_LOG(LogRivermax, Warning, TEXT("No chunks ready to skip. Waiting"));
+								TRACE_CPUPROFILER_EVENT_SCOPE(PrepareNextFrame::NoFreeChunk);
+								bHasAddedTrace = true;
+							}
+						}
+						else
+						{
+							ensure(false);
+							UE_LOG(LogRivermax, Error, TEXT("Invalid error happened while trying to skip chunks. Status: %d."), Status);
+							Listener->OnStreamError();
+							Stop();
 						}
 					}
-					else
-					{
-						ensure(false);
-						UE_LOG(LogRivermax, Error, TEXT("Invalid error happened while trying to skip chunks. Status: %d."), Status);
-						Listener->OnStreamError();
-						Stop();
-					}
-				}
-			} while (Status != RMAX_OK && bIsActive);
+				} while (Status != RMAX_OK && bIsActive);
+
+				// Keep RTP header in sync when skipping chunks
+				StreamData.RTPHeaderFrameIndex = CurrentFrame->FrameIndex;
+			}
 		}
 	}
 
@@ -1163,7 +1331,7 @@ namespace UE::RivermaxCore::Private
 			if (CurrentTime - LastStatsShownTimestamp > CVarRivermaxOutputShowStatsInterval.GetValueOnAnyThread())
 			{
 				LastStatsShownTimestamp = CurrentTime;
-				UE_LOG(LogRivermax, Log, TEXT("Stats: FrameSent: %u. CommitImmediate: %u. CommitRetries: %u. ChunkRetries: %u. ChunkSkippingRetries: %u"), Stats.MemoryBlockSentCounter, Stats.CommitImmediate, Stats.CommitRetries, Stats.ChunkRetries, Stats.ChunkSkippingRetries);
+				UE_LOG(LogRivermax, Log, TEXT("Stats: FrameSent: %u. CommitImmediate: %u. CommitRetries: %u. ChunkRetries: %u. ChunkSkippingRetries: %u"), Stats.FramesSentCounter, Stats.CommitImmediate, Stats.CommitRetries, Stats.TotalChunkRetries, Stats.ChunkSkippingRetries);
 			}
 		}
 	}
@@ -1191,24 +1359,62 @@ namespace UE::RivermaxCore::Private
 			}
 		}
 
+		// Work around when dealing with multi memblocks. Rivermax fails to create stream 
+		// With memblock not starting on the right cuda alignment. 
+		const bool bAlignEachFrameMemory = !CachedCVars.bUseSingleMemblock;
 		FFrameManagerSetupArgs FrameManagerArgs;
 		FrameManagerArgs.Resolution = Options.AlignedResolution;
 		FrameManagerArgs.bTryGPUAllocation = bTryGPUDirect;
 		FrameManagerArgs.NumberOfFrames = Options.NumberOfBuffers;
 		FrameManagerArgs.Stride = GetStride();
+		FrameManagerArgs.FrameDesiredSize = StreamMemory.ChunksPerFrameField * StreamMemory.PacketsPerChunk * StreamMemory.PayloadSize;
+		FrameManagerArgs.bAlignEachFrameAlloc = bAlignEachFrameMemory;
 		FrameManagerArgs.OnFreeFrameDelegate = FOnFrameReadyDelegate::CreateRaw(this, &FRivermaxOutputStream::OnFrameReadyToBeUsed);
 		FrameManagerArgs.OnPreFrameReadyDelegate = FOnFrameReadyDelegate::CreateRaw(this, &FRivermaxOutputStream::OnPreFrameReadyToBeSent);
 		FrameManagerArgs.OnFrameReadyDelegate = FOnFrameReadyDelegate::CreateRaw(this, &FRivermaxOutputStream::OnFrameReadyToBeSent);
+		FrameManagerArgs.OnCriticalErrorDelegate = FOnCriticalErrorDelegate::CreateRaw(this, &FRivermaxOutputStream::OnFrameManagerCriticalError);
 		const EFrameMemoryLocation FrameLocation = FrameManager->Initialize(FrameManagerArgs);
 		bUseGPUDirect = FrameLocation == EFrameMemoryLocation::GPU;
 
+
+		const bool bHasAllocatedFrames = (FrameLocation != EFrameMemoryLocation::None);
+		if (bHasAllocatedFrames && CachedCVars.bUseIntermediateBuffer)
+		{
+			// Allocate intermediate buffer in same memory space as frame memory.
+			FOnFrameDataCopiedDelegate OnDataCopiedDelegate = FOnFrameDataCopiedDelegate::CreateRaw(this, &FRivermaxOutputStream::OnMemoryChunksCopied);
+
+			const int32 DesiredSize = FrameManagerArgs.FrameDesiredSize * StreamMemory.FramesFieldPerMemoryBlock;
+			if (bUseGPUDirect)
+			{
+				Allocator = MakeUnique<FGPUAllocator>(DesiredSize, OnDataCopiedDelegate);
+			}
+			else
+			{
+				Allocator = MakeUnique<FSystemAllocator>(DesiredSize, OnDataCopiedDelegate);
+			}
+			
+			if (!Allocator->Allocate(StreamMemory.MemoryBlockCount, bAlignEachFrameMemory))
+			{
+				return false;
+			}
+		}
+
 		return FrameLocation != EFrameMemoryLocation::None;
 	}
 
 	void FRivermaxOutputStream::CleanupFrameManagement()
 	{
-		FrameManager->Cleanup();
-		FrameManager.Reset();
+		if (FrameManager)
+		{
+			FrameManager->Cleanup();
+			FrameManager.Reset();
+		}
+
+		if (Allocator)
+		{
+			Allocator->Deallocate();
+			Allocator.Reset();
+		}
 	}
 
 	int32 FRivermaxOutputStream::GetStride() const
@@ -1303,16 +1509,17 @@ namespace UE::RivermaxCore::Private
 
 	bool FRivermaxOutputStream::ReserveFrame(uint32 FrameIdentifier) const
 	{
-		TSharedPtr<FRivermaxOutputFrame> ReservedFrame = FrameManager->GetFreeFrame();
+		TSharedPtr<FRivermaxOutputFrame> ReservedFrame = FrameManager->GetNextFrame(FrameIdentifier);
 		if (!ReservedFrame && Options.FrameLockingMode == EFrameLockingMode::BlockOnReservation)
 		{
 			TRACE_CPUPROFILER_EVENT_SCOPE(Rmax::WaitForAvailableFrame);
 			while(!ReservedFrame && bIsActive)
 			{
 				FrameAvailableSignal->Wait();
-				ReservedFrame = FrameManager->GetFreeFrame();
+				ReservedFrame = FrameManager->GetNextFrame(FrameIdentifier);
 			}
 		}
+
 		return ReservedFrame != nullptr;
 	}
 
@@ -1337,5 +1544,70 @@ namespace UE::RivermaxCore::Private
 		Stop();
 	}
 
+	void FRivermaxOutputStream::CacheCVarValues()
+	{
+		CachedCVars.ForceCommitImmediateTimeNanosec = CVarRivermaxOutputForceImmediateSchedulingThreshold.GetValueOnAnyThread();
+		CachedCVars.SkipSchedulingTimeNanosec = CVarRivermaxOutputSkipSchedulingCutOffTime.GetValueOnAnyThread() * 1E3;
+		CachedCVars.FrameMemorySliceCount = FMath::Clamp(CVarRivermaxOutputFrameSliceCount.GetValueOnAnyThread(), 1, 100);
+		CachedCVars.bUseIntermediateBuffer = CVarRivermaxOutputEnableIntermediateBuffer.GetValueOnAnyThread() == 1;
+		CachedCVars.bUseSingleMemblock = CVarRivermaxOutputUseSingleMemblock.GetValueOnAnyThread() == 1;
+	}
+
+	bool FRivermaxOutputStream::CanProceedWithScheduling() const
+	{
+		switch (Options.AlignmentMode)
+		{
+		case ERivermaxAlignmentMode::AlignmentPoint:
+		{
+			return CanProceedWithScheduling_AlignmentPoints();
+		}
+		case ERivermaxAlignmentMode::FrameCreation:
+		{
+			return CanProceedWithScheduling_FrameCreation();
+		}
+		default:
+		{
+			checkNoEntry();
+			return false;
+		}
+		}
+
+		return false;
+	}
+
+	bool FRivermaxOutputStream::CanProceedWithScheduling_FrameCreation() const
+	{
+		return true;
+	}
+
+	bool FRivermaxOutputStream::CanProceedWithScheduling_AlignmentPoints() const
+	{
+		// When aligning on frame boundary, we need to make sure we get here before the scheduling time
+		// There is also some configurable buffer added to make sure we are not right on time
+		// which would cause timing issues.
+		const uint64 CurrentTime = RivermaxModule->GetRivermaxManager()->GetTime();
+		if (StreamData.NextScheduleTimeNanosec <= (CurrentTime + CachedCVars.SkipSchedulingTimeNanosec))
+		{
+			UE_LOG(LogRivermax, Warning, TEXT("Can't schedule next frame. Cut off time already passed."));
+			return false;
+		}
+
+		// Add other causes of timing issues. 
+		// Possible options : Chunk warnings, Last commit time too close to frame boundary, etc...
+		
+		return true;
+	}
+
+	void FRivermaxOutputStream::OnMemoryChunksCopied(const TSharedPtr<FBaseDataCopySideCar>& Sidecar)
+	{
+		if (Sidecar)
+		{
+			TSharedPtr<FMemChunkCopiedInfo> CopyInfo = StaticCastSharedPtr<FMemChunkCopiedInfo>(Sidecar);
+			if (CurrentFrame)
+			{
+				CurrentFrame->ScheduleTimeCopied = CopyInfo->TimeCovered;
+			}
+		}
+	}
 }
 
diff --git a/Engine/Plugins/VirtualProduction/Rivermax/RivermaxCore/Source/RivermaxCore/Private/RivermaxOutputStream.h b/Engine/Plugins/VirtualProduction/Rivermax/RivermaxCore/Source/RivermaxCore/Private/RivermaxOutputStream.h
--- a/Engine/Plugins/VirtualProduction/Rivermax/RivermaxCore/Source/RivermaxCore/Private/RivermaxOutputStream.h
+++ b/Engine/Plugins/VirtualProduction/Rivermax/RivermaxCore/Source/RivermaxCore/Private/RivermaxOutputStream.h
@@ -19,6 +19,8 @@ class IRivermaxCoreModule;
 namespace UE::RivermaxCore::Private
 {
 	class FFrameManager;
+	class FBaseFrameAllocator;
+	struct FBaseDataCopySideCar;
 
 	using UE::RivermaxCore::FRivermaxOutputStreamOptions;
 	using UE::RivermaxCore::FRivermaxOutputVideoFrameInfo;
@@ -46,21 +48,20 @@ namespace UE::RivermaxCore::Private
 		TArray<rmax_mem_block> MemoryBlocks;
 		TArray<uint16_t> PayloadSizes; //Array describing stride payload size
 		TArray<uint16_t> HeaderSizes; //Array describing header payload size
-		TArray<TArray<FRawRTPHeader>> RTPHeaders;
-
-		rmax_buffer_attr BufferAttributes;
+		TArray<TArray<FRawRTPHeader>> RTPHeaders; //RTP Headers per memblock
 	};
 
 	struct FRivermaxOutputStreamStats
 	{
-		uint32 ChunkRetries = 0;
+		uint32 TotalChunkRetries = 0;
+		uint32 LastFrameChunkRetries = 0;
 		uint32 ChunkSkippingRetries = 0;
 		uint32 TotalStrides = 0;
 		uint32 ChunkWait = 0;
 		uint32 CommitWaits = 0;
 		uint32 CommitRetries = 0;
 		uint32 CommitImmediate = 0;
-		uint64 MemoryBlockSentCounter = 0;
+		uint64 FramesSentCounter = 0;
 	};
 
 	struct FRivermaxOutputStreamData
@@ -68,6 +69,9 @@ namespace UE::RivermaxCore::Private
 		/** Current sequence number being done */
 		uint32 SequenceNumber = 0;
 		double FrameFieldTimeIntervalNs = 0.0;
+		
+		/** RTP header index we are expected to use. Used to validate returned header pointer by Rivermax. */
+		uint8 RTPHeaderFrameIndex = 0;
 
 		/** Used to detect misalignment between chunk being sent and frame memory we are writing in */
 		bool bHasFrameFirstChunkBeenFetched = false;
@@ -86,6 +90,32 @@ namespace UE::RivermaxCore::Private
 
 		/** Timestamp at which we started commiting a frame */
 		uint64 LastSendStartTimeNanoSec = 0;
+		
+		/** Keeping track of how much time was slept last round. */
+		uint64 LastSleepTimeNanoSec = 0;
+	};
+
+	/** Struct holding various cached cvar values that can't be changed once stream has been created and to avoid calling anythread getters continuously */
+	struct FOutputStreamCachedCVars
+	{
+		/** Time padding from scheduling time required to avoid skipping it */
+		uint64 SkipSchedulingTimeNanosec = 0;
+
+		/** 
+		 * Time from scheduling required to not commit it immediately 
+		 * Rivermax sdk will throw an error if time is in the past when it
+		 * gets to actually comitting it. 
+		 */
+		uint64 ForceCommitImmediateTimeNanosec = 0;
+
+		/** Number of slices we split frame data into when copying it into intermediate buffer */
+		uint32 FrameMemorySliceCount = 1;
+		
+		/** Whether intermediate buffer is used and captured frame has to be copied over again. */
+		bool bUseIntermediateBuffer = true;
+		
+		/** Tentative optimization recommended for SDK where a single big memblock is allocated. When false, a memblock per frame is configured. */
+		bool bUseSingleMemblock = true;
 	};
 
 	class FRivermaxOutputStream : public UE::RivermaxCore::IRivermaxOutputStream, public FRunnable
@@ -145,9 +175,21 @@ namespace UE::RivermaxCore::Private
 		/** Calculate next frame scheduling time for frame creation mode */
 		void CalculateNextScheduleTime_FrameCreation(uint64 CurrentClockTimeNanosec, uint64 CurrentFrameNumber);
 
+		/** Validates timing post waiting to be sure it is good moving forward with the next frame */
+		bool CanProceedWithScheduling() const;
+		
+		/** Scheduling verification for frame creation mode. Always valid for now. */
+		bool CanProceedWithScheduling_FrameCreation() const;
+		
+		/** Scheduling verification for alignment points which validates current timing is ok with alignment points */
+		bool CanProceedWithScheduling_AlignmentPoints() const;
+
 		/** Query rivermax library for the next chunk to work with */
 		void GetNextChunk();
 
+		/** Copies part of frame memory in next memblock's chunk to be sent out */
+		void CopyFrameData();
+
 		/** Fills RTP header for all packets to be sent for this chunk */
 		void SetupRTPHeaders();
 
@@ -184,6 +226,11 @@ namespace UE::RivermaxCore::Private
 		/** Used to detect when the frame manager has caught a critical error */
 		void OnFrameManagerCriticalError();
 
+		/** Used to cache cvars at initialization */
+		void CacheCVarValues();
+
+		/** Called back when copy request was completed by allocator */
+		void OnMemoryChunksCopied(const TSharedPtr<FBaseDataCopySideCar>& Sidecar);
 
 	private:
 
@@ -214,6 +261,9 @@ namespace UE::RivermaxCore::Private
 		/** Manages allocation and memory manipulation of video frames */
 		TUniquePtr<FFrameManager> FrameManager;
 
+		/** Manages allocation of memory for rivermax memblocks */
+		TUniquePtr<FBaseFrameAllocator> Allocator;
+
 		/** Whether stream is active or not */
 		std::atomic<bool> bIsActive;
 
@@ -250,6 +300,9 @@ namespace UE::RivermaxCore::Private
 		/** Future returned by the async initialization job we launch. Used to detect if it has completed during shutdown. */
 		TFuture<void> InitializationFuture;
 
+		/** Cached cvar values */
+		FOutputStreamCachedCVars CachedCVars;
+
 		/** Time to sleep when waiting for an operation to complete */
 		static constexpr double SleepTimeSeconds = 50.0 * 1E-6;
 	};
