Fix static analysis error in RivermaxManager.cpp

diff --git a/Engine/Plugins/VirtualProduction/Rivermax/RivermaxCore/Source/RivermaxCore/Private/RivermaxManager.cpp b/Engine/Plugins/VirtualProduction/Rivermax/RivermaxCore/Source/RivermaxCore/Private/RivermaxManager.cpp
--- a/Engine/Plugins/VirtualProduction/Rivermax/RivermaxCore/Source/RivermaxCore/Private/RivermaxManager.cpp
+++ b/Engine/Plugins/VirtualProduction/Rivermax/RivermaxCore/Source/RivermaxCore/Private/RivermaxManager.cpp
@@ -203,7 +203,6 @@ namespace UE::RivermaxCore::Private
 			&& FuncList.rmax_set_device_config
 			&& FuncList.rmax_unset_device_config
 			&& FuncList.rmax_request_notification
-			&& FuncList.rmax_request_notification
 			&& FuncList.rmax_get_event_channel
 			&& FuncList.rmax_register_memory
 			&& FuncList.rmax_register_memory_ex
