2110 ptp / stream alignment bug fix
Inner frustum streaming corruption bug fix

- Moving rmax late commit time protection closer to actual commit and making detection active the entire frame
- Changing cut off threshold default value to not be too restrictive
- Removing some clutter in rmax output hot loop but mainly removed all media capture pipe insights dynamic texts.
- Added RTP header prefill to save time when sending
- Added early copy of frame when it's ready to reduce risk of timing issues
- Disabled intermediate buffer for frame creation alignment to address the corruption issue

diff --git a/Engine/Plugins/Media/MediaIOFramework/Source/MediaIOCore/Private/MediaCapture.cpp b/Engine/Plugins/Media/MediaIOFramework/Source/MediaIOCore/Private/MediaCapture.cpp
--- a/Engine/Plugins/Media/MediaIOFramework/Source/MediaIOCore/Private/MediaCapture.cpp
+++ b/Engine/Plugins/Media/MediaIOFramework/Source/MediaIOCore/Private/MediaCapture.cpp
@@ -639,7 +639,7 @@ void UMediaCapture::CaptureImmediate_RenderThread(const UE::MediaCaptureData::FC
 {
 	using namespace UE::MediaCaptureData;
 	TRACE_CPUPROFILER_EVENT_SCOPE(UMediaCapture::CaptureImmediate_RenderThread);
-	TRACE_CPUPROFILER_EVENT_SCOPE_TEXT(*FString::Printf(TEXT("MediaCapturePipe: %llu"), GFrameCounterRenderThread));
+	TRACE_CPUPROFILER_EVENT_SCOPE_TEXT(*FString::Printf(TEXT("MediaCapturePipe: %llu"), GFrameCounterRenderThread % 10));
 	
 	check(IsInRenderingThread());
 	
@@ -1274,7 +1274,8 @@ void UMediaCapture::OnEndFrame_GameThread()
 {
 	using namespace UE::MediaCaptureData;
 
-	TRACE_CPUPROFILER_EVENT_SCOPE_TEXT(*FString::Printf(TEXT("MediaCapture End Frame %d"), GFrameCounter));
+	TRACE_CPUPROFILER_EVENT_SCOPE(MediaCaptureEndFrame);
+	TRACE_CPUPROFILER_EVENT_SCOPE_TEXT(*FString::Printf(TEXT("MediaCapturePipe: %llu"), GFrameCounter % 10));
 
 	if (!bOutputResourcesInitialized)
 	{
@@ -1327,7 +1328,7 @@ bool UMediaCapture::ProcessCapture_RenderThread(const TSharedPtr<UE::MediaCaptur
 	{
 		FrameNumber = CapturingFrame->CaptureBaseData.SourceFrameNumberRenderThread;
 	}
-	TRACE_CPUPROFILER_EVENT_SCOPE_TEXT(*FString::Printf(TEXT("Process Capture Render Thread Frame %d"), FrameNumber));
+	TRACE_CPUPROFILER_EVENT_SCOPE_TEXT(*FString::Printf(TEXT("Process Capture Render Thread Frame %d"), FrameNumber % 10));
 
 	if (CapturingFrame)
 	{
@@ -1451,7 +1452,7 @@ bool UMediaCapture::ProcessReadyFrame_RenderThread(FRHICommandListImmediate& RHI
 			{
 				{
 					SCOPE_CYCLE_COUNTER(STAT_MediaCapture_RenderThread_RHI_CaptureCallback)
-					TRACE_CPUPROFILER_EVENT_SCOPE_TEXT(*FString::Printf(TEXT("MediaCapture Output Frame %d"), ReadyFrame->CaptureBaseData.SourceFrameNumberRenderThread));
+					TRACE_CPUPROFILER_EVENT_SCOPE_TEXT(*FString::Printf(TEXT("MediaCapture Output Frame %d"), ReadyFrame->CaptureBaseData.SourceFrameNumberRenderThread % 10));
 
 					// The Width/Height of the surface may be different then the DesiredOutputSize : Some underlying implementations enforce a specific stride, therefore
 					// there may be padding at the end of each row.
@@ -1479,7 +1480,7 @@ bool UMediaCapture::ProcessReadyFrame_RenderThread(FRHICommandListImmediate& RHI
 				}
 
 				TRACE_CPUPROFILER_EVENT_SCOPE(UMediaCapture::RHIResourceCaptured);
-				TRACE_CPUPROFILER_EVENT_SCOPE_TEXT(*FString::Printf(TEXT("MediaCapture Output Frame %d"), ReadyFrame->CaptureBaseData.SourceFrameNumberRenderThread));
+				TRACE_CPUPROFILER_EVENT_SCOPE_TEXT(*FString::Printf(TEXT("MediaCapture Output Frame %d"), ReadyFrame->CaptureBaseData.SourceFrameNumberRenderThread % 10));
 				SCOPE_CYCLE_COUNTER(STAT_MediaCapture_RenderThread_CaptureCallback)
 					InMediaCapture->OnRHIResourceCaptured_RenderingThread(ReadyFrame->CaptureBaseData, ReadyFrame->UserData, ReadyFrame->GetTextureResource());
 			}
diff --git a/Engine/Plugins/Media/MediaIOFramework/Source/MediaIOCore/Private/MediaCaptureHelper.cpp b/Engine/Plugins/Media/MediaIOFramework/Source/MediaIOCore/Private/MediaCaptureHelper.cpp
--- a/Engine/Plugins/Media/MediaIOFramework/Source/MediaIOCore/Private/MediaCaptureHelper.cpp
+++ b/Engine/Plugins/Media/MediaIOFramework/Source/MediaIOCore/Private/MediaCaptureHelper.cpp
@@ -131,7 +131,7 @@ bool FMediaCaptureHelper::CaptureFrame(const UE::MediaCaptureData::FCaptureFrame
 	if (CapturingFrame->IsTextureResource())
 	{
 		TRACE_CPUPROFILER_EVENT_SCOPE(UMediaCapture::LockDMATexture_RenderThread);
-		TRACE_CPUPROFILER_EVENT_SCOPE_TEXT(*FString::Printf(TEXT("LockDmaTexture Output Frame %d"), CapturingFrame->CaptureBaseData.SourceFrameNumberRenderThread));
+		TRACE_CPUPROFILER_EVENT_SCOPE_TEXT(*FString::Printf(TEXT("LockDmaTexture Output Frame %d"), CapturingFrame->CaptureBaseData.SourceFrameNumberRenderThread % 10));
 		Args.MediaCapture->LockDMATexture_RenderThread(CapturingFrame->GetTextureResource());
 	}
 
@@ -254,7 +254,7 @@ void FMediaCaptureHelper::OnReadbackComplete(FRHICommandList& RHICmdList, UMedia
 					}
 
 					TRACE_CPUPROFILER_EVENT_SCOPE(UMediaCapture::RHIResourceCaptured);
-					TRACE_CPUPROFILER_EVENT_SCOPE_TEXT(*FString::Printf(TEXT("MediaCapture Output Frame %d"), ReadyFrame->CaptureBaseData.SourceFrameNumberRenderThread));
+					TRACE_CPUPROFILER_EVENT_SCOPE_TEXT(*FString::Printf(TEXT("MediaCapture Output Frame %d"), ReadyFrame->CaptureBaseData.SourceFrameNumberRenderThread % 10));
 					SCOPE_CYCLE_COUNTER(STAT_MediaCaptureHelper_RenderThread_RHI_CaptureCallback)
 
 						MediaCapture->OnRHIResourceCaptured_RenderingThread(ReadyFrame->CaptureBaseData, ReadyFrame->UserData, ReadyFrame->GetTextureResource());
@@ -278,7 +278,7 @@ void FMediaCaptureHelper::OnReadbackComplete(FRHICommandList& RHICmdList, UMedia
 				if (ensure(ColorDataBuffer))
 				{
 					{
-						TRACE_CPUPROFILER_EVENT_SCOPE_TEXT(*FString::Printf(TEXT("MediaCapture Output Frame %d"), ReadyFrame->GetId()));
+						TRACE_CPUPROFILER_EVENT_SCOPE_TEXT(*FString::Printf(TEXT("MediaCapture Output Frame %d"), ReadyFrame->GetId() % 10));
 						SCOPE_CYCLE_COUNTER(STAT_MediaCaptureHelper_AnyThread_LockResource)
 							MediaCapture->OnFrameCaptured_RenderingThread(ReadyFrame->CaptureBaseData, ReadyFrame->UserData, ColorDataBuffer, MediaCapture->DesiredOutputSize.X, MediaCapture->DesiredOutputSize.Y, RowStride);
 					}
@@ -300,7 +300,7 @@ void FMediaCaptureHelper::OnReadbackComplete(FRHICommandList& RHICmdList, UMedia
 					}
 
 					TRACE_CPUPROFILER_EVENT_SCOPE(UMediaCapture::RHIResourceCaptured);
-					TRACE_CPUPROFILER_EVENT_SCOPE_TEXT(*FString::Printf(TEXT("MediaCapture Output Frame %d"), ReadyFrame->CaptureBaseData.SourceFrameNumberRenderThread));
+					TRACE_CPUPROFILER_EVENT_SCOPE_TEXT(*FString::Printf(TEXT("MediaCapture Output Frame %d"), ReadyFrame->CaptureBaseData.SourceFrameNumberRenderThread % 10));
 					SCOPE_CYCLE_COUNTER(STAT_MediaCaptureHelper_RenderThread_RHI_CaptureCallback)
 
 					MediaCapture->OnRHIResourceCaptured_AnyThread(ReadyFrame->CaptureBaseData, ReadyFrame->UserData, ReadyFrame->GetTextureResource());
@@ -324,7 +324,7 @@ void FMediaCaptureHelper::OnReadbackComplete(FRHICommandList& RHICmdList, UMedia
 				if (ensure(ColorDataBuffer))
 				{
 					{
-						TRACE_CPUPROFILER_EVENT_SCOPE_TEXT(*FString::Printf(TEXT("MediaCapture Output Frame %d"), ReadyFrame->GetId()));
+						TRACE_CPUPROFILER_EVENT_SCOPE_TEXT(*FString::Printf(TEXT("MediaCapture Output Frame %d"), ReadyFrame->GetId() % 10));
 						SCOPE_CYCLE_COUNTER(STAT_MediaCaptureHelper_AnyThread_LockResource)
 
 						UMediaCapture::FMediaCaptureResourceData ResourceData;
diff --git a/Engine/Plugins/Media/MediaIOFramework/Source/MediaIOCore/Private/MediaCaptureHelper.h b/Engine/Plugins/Media/MediaIOFramework/Source/MediaIOCore/Private/MediaCaptureHelper.h
--- a/Engine/Plugins/Media/MediaIOFramework/Source/MediaIOCore/Private/MediaCaptureHelper.h
+++ b/Engine/Plugins/Media/MediaIOFramework/Source/MediaIOCore/Private/MediaCaptureHelper.h
@@ -32,7 +32,7 @@ namespace UE::MediaCaptureData
 		static void AddSyncPointPass(FRDGBuilder& GraphBuilder, UMediaCapture* MediaCapture, TSharedPtr<UE::MediaCaptureData::FCaptureFrame> CapturingFrame, FRDGViewableResource* OutputResource)
 		{
 			TRACE_CPUPROFILER_EVENT_SCOPE(MediaCaptureSyncPoint);
-			TRACE_CPUPROFILER_EVENT_SCOPE_TEXT(*FString::Printf(TEXT("MediaCaptureSyncPoint_%d"), CapturingFrame->CaptureBaseData.SourceFrameNumberRenderThread));
+			TRACE_CPUPROFILER_EVENT_SCOPE_TEXT(*FString::Printf(TEXT("MediaCaptureSyncPoint_%d"), CapturingFrame->CaptureBaseData.SourceFrameNumberRenderThread % 10));
 
 			// Initialize sync handlers only the first time. 
 			if (MediaCapture->bSyncHandlersInitialized == false)
diff --git a/Engine/Plugins/VirtualProduction/Rivermax/RivermaxCore/Source/RivermaxCore/Private/RivermaxFrameAllocator.cpp b/Engine/Plugins/VirtualProduction/Rivermax/RivermaxCore/Source/RivermaxCore/Private/RivermaxFrameAllocator.cpp
--- a/Engine/Plugins/VirtualProduction/Rivermax/RivermaxCore/Source/RivermaxCore/Private/RivermaxFrameAllocator.cpp
+++ b/Engine/Plugins/VirtualProduction/Rivermax/RivermaxCore/Source/RivermaxCore/Private/RivermaxFrameAllocator.cpp
@@ -11,12 +11,12 @@
 namespace UE::RivermaxCore::Private
 {
 	static TAutoConsoleVariable<int32> CVarRivermaxOutputEnableParallelCopy(
-		TEXT("Rivermax.Output.EnableParallelCopy"), 0,
+		TEXT("Rivermax.Output.EnableParallelCopy"), 1,
 		TEXT("Parallel"),
 		ECVF_Default);
 
 	static TAutoConsoleVariable<int32> CVarRivermaxOutputParallelCopyThreadCount(
-		TEXT("Rivermax.Output.ParallelCopyThreadCount"), 8,
+		TEXT("Rivermax.Output.ParallelCopyThreadCount"), 4,
 		TEXT("Parallel"),
 		ECVF_Default);
 
diff --git a/Engine/Plugins/VirtualProduction/Rivermax/RivermaxCore/Source/RivermaxCore/Private/RivermaxFrameManager.cpp b/Engine/Plugins/VirtualProduction/Rivermax/RivermaxCore/Source/RivermaxCore/Private/RivermaxFrameManager.cpp
--- a/Engine/Plugins/VirtualProduction/Rivermax/RivermaxCore/Source/RivermaxCore/Private/RivermaxFrameManager.cpp
+++ b/Engine/Plugins/VirtualProduction/Rivermax/RivermaxCore/Source/RivermaxCore/Private/RivermaxFrameManager.cpp
@@ -5,6 +5,7 @@
 #include "IRivermaxCoreModule.h"
 #include "IRivermaxManager.h"
 #include "RivermaxLog.h"
+#include "RivermaxTracingUtils.h"
 
 namespace UE::RivermaxCore::Private
 {
@@ -238,9 +239,8 @@ namespace UE::RivermaxCore::Private
 				AvailableFrame->bIsVideoBufferReady = true;
 				if (AvailableFrame->IsReadyToBeSent())
 				{
-					const FString TraceName = FString::Format(TEXT("Rmax::FrameReady {0}"), { AvailableFrame->FrameIndex });
-					TRACE_CPUPROFILER_EVENT_SCOPE_TEXT(*TraceName);
-					TRACE_CPUPROFILER_EVENT_SCOPE_TEXT(*FString::Printf(TEXT("MediaCapturePipe: %u"), AvailableFrame->FrameIdentifier));
+					TRACE_CPUPROFILER_EVENT_SCOPE_TEXT(*FRivermaxTracingUtils::RmaxOutFrameReadyTraceEvents[AvailableFrame->FrameIndex % 10]);
+					TRACE_CPUPROFILER_EVENT_SCOPE_TEXT(*FRivermaxTracingUtils::RmaxOutMediaCapturePipeTraceEvents[AvailableFrame->FrameIdentifier % 10]);
 
 					MarkAsReady(AvailableFrame);
 				}
diff --git a/Engine/Plugins/VirtualProduction/Rivermax/RivermaxCore/Source/RivermaxCore/Private/RivermaxInputStream.cpp b/Engine/Plugins/VirtualProduction/Rivermax/RivermaxCore/Source/RivermaxCore/Private/RivermaxInputStream.cpp
--- a/Engine/Plugins/VirtualProduction/Rivermax/RivermaxCore/Source/RivermaxCore/Private/RivermaxInputStream.cpp
+++ b/Engine/Plugins/VirtualProduction/Rivermax/RivermaxCore/Source/RivermaxCore/Private/RivermaxInputStream.cpp
@@ -9,6 +9,7 @@
 #include "IRivermaxManager.h"
 #include "Misc/ByteSwap.h"
 #include "RivermaxLog.h"
+#include "RivermaxTracingUtils.h"
 #include "RivermaxUtils.h"
 #include "RTPHeader.h"
 
@@ -30,18 +31,6 @@ namespace UE::RivermaxCore::Private
 		TEXT("Expected payload size used to initialize rivermax stream."),
 		ECVF_Default);
 
-	/** 
-	 * Converts a timestamp in MediaClock period units to a frame number for a given frame rate 
-	 * 2110-20 streams uses a standard media clock rate of 90kHz
-	 */
-	uint32 TimestampToFrameNumber(uint32 Timestamp, const FFrameRate& FrameRate)
-	{
-		using namespace UE::RivermaxCore::Private::Utils;
-		const double MediaFrameTime = Timestamp / MediaClockSampleRate;
-		const uint32 FrameNumber = FMath::Floor(MediaFrameTime * FrameRate.AsDecimal());
-		return FrameNumber;
-	}
-
 	void FFrameDescriptionTrackingData::ResetSingleFrameTracking()
 	{
 		PayloadSizeReceived.Empty();
@@ -382,8 +371,8 @@ namespace UE::RivermaxCore::Private
 					// Add trace for the first packet of a frame to help visualize reception of a full frame in time
 					if (bIsFirstPacketReceived == false)
 					{
-						const FString TraceName = FString::Format(TEXT("RmaxInput::StartingFrame {0}"), { RTPHeader.Timestamp });
-						TRACE_CPUPROFILER_EVENT_SCOPE_TEXT(*TraceName);
+						const uint32 FrameNumber = Utils::TimestampToFrameNumber(RTPHeader.Timestamp, Options.FrameRate);
+						TRACE_CPUPROFILER_EVENT_SCOPE_TEXT(*FRivermaxTracingUtils::RmaxInStartingFrameTraceEvents[FrameNumber % 10]);
 						bIsFirstPacketReceived = true;
 					}
 
@@ -861,14 +850,13 @@ namespace UE::RivermaxCore::Private
 		Descriptor.Timestamp = RTPHeader.Timestamp;
 		const uint32 PixelCount = StreamResolution.X * StreamResolution.Y;
 		Descriptor.VideoBufferSize = PixelCount / FormatInfo.PixelGroupCoverage * FormatInfo.PixelGroupSize;
-		Descriptor.FrameNumber = TimestampToFrameNumber(RTPHeader.Timestamp, Options.FrameRate);
+		Descriptor.FrameNumber = Utils::TimestampToFrameNumber(RTPHeader.Timestamp, Options.FrameRate);
 
 		if (StreamData.ReceivedSize == StreamData.ExpectedSize)
 		{
 			++StreamStats.FramesReceived;
 			
-			const FString TraceName = FString::Format(TEXT("RmaxReceivedFrame {0}|{1}"), { RTPHeader.Timestamp, Descriptor.FrameNumber });
-			TRACE_CPUPROFILER_EVENT_SCOPE_TEXT(*TraceName);
+			TRACE_CPUPROFILER_EVENT_SCOPE_TEXT(*FRivermaxTracingUtils::RmaxInReceivedFrameTraceEvents[Descriptor.FrameNumber % 10]);
 
 			if (bIsUsingGPUDirect)
 			{
@@ -1022,7 +1010,7 @@ namespace UE::RivermaxCore::Private
 		const uint32 PixelCount = StreamResolution.X * StreamResolution.Y;
 		const uint32 FrameSize = PixelCount / FormatInfo.PixelGroupCoverage * FormatInfo.PixelGroupSize;
 		Descriptor.VideoBufferSize = FrameSize;
-		Descriptor.FrameNumber = TimestampToFrameNumber(RTPHeader.Timestamp, Options.FrameRate);
+		Descriptor.FrameNumber = Utils::TimestampToFrameNumber(RTPHeader.Timestamp, Options.FrameRate);
 
 		const uint64 LastSequenceNumberIncremented = StreamData.LastSequenceNumber + 1;
 
diff --git a/Engine/Plugins/VirtualProduction/Rivermax/RivermaxCore/Source/RivermaxCore/Private/RivermaxManager.cpp b/Engine/Plugins/VirtualProduction/Rivermax/RivermaxCore/Source/RivermaxCore/Private/RivermaxManager.cpp
--- a/Engine/Plugins/VirtualProduction/Rivermax/RivermaxCore/Source/RivermaxCore/Private/RivermaxManager.cpp
+++ b/Engine/Plugins/VirtualProduction/Rivermax/RivermaxCore/Source/RivermaxCore/Private/RivermaxManager.cpp
@@ -11,6 +11,7 @@
 #include "Misc/CoreDelegates.h"
 #include "RivermaxDeviceFinder.h"
 #include "RivermaxLog.h"
+#include "RivermaxTracingUtils.h"
 
 #include <chrono>
 
@@ -30,6 +31,16 @@
 
 #define LOCTEXT_NAMESPACE "RivermaxManager"
 
+namespace UE::RivermaxCore
+{
+	TMap<uint8, FString> FRivermaxTracingUtils::RmaxOutMediaCapturePipeTraceEvents;
+	TMap<uint8, FString> FRivermaxTracingUtils::RmaxOutSendingFrameTraceEvents;
+	TMap<uint8, FString> FRivermaxTracingUtils::RmaxOutFrameReadyTraceEvents;
+	TMap<uint8, FString> FRivermaxTracingUtils::RmaxInStartingFrameTraceEvents;
+	TMap<uint8, FString> FRivermaxTracingUtils::RmaxInReceivedFrameTraceEvents;
+	TMap<uint8, FString> FRivermaxTracingUtils::RmaxInSelectedFrameTraceEvents;
+}
+
 namespace UE::RivermaxCore::Private
 {
 	static TAutoConsoleVariable<int32> CVarRivermaxEnableGPUDirectCapability(
@@ -293,6 +304,8 @@ namespace UE::RivermaxCore::Private
 							VerifyGPUDirectCapability();
 						}
 
+						InitializeTraceMarkupStrings();
+
 						bIsLibraryInitialized = true;
 						const TCHAR* GPUDirectSupport = bIsGPUDirectSupported ? TEXT("with GPUDirect support.") : TEXT("without GPUDirect support.");
 						UE_LOG(LogRivermax, Log, TEXT("Rivermax library version %d.%d.%d.%d succesfully initialized, %s"), Major, Minor, Release, Build, GPUDirectSupport);
@@ -669,6 +682,21 @@ namespace UE::RivermaxCore::Private
 		}
 	}
 
+	void FRivermaxManager::InitializeTraceMarkupStrings()
+	{
+		constexpr uint8 MarkupCount = 10;
+
+		for (uint8 Index = 0; Index < MarkupCount; ++Index)
+		{
+			FRivermaxTracingUtils::RmaxOutMediaCapturePipeTraceEvents.Add(Index, FString::Printf(TEXT("MediaCapturePipe: %u"), Index));
+			FRivermaxTracingUtils::RmaxOutSendingFrameTraceEvents.Add(Index, FString::Printf(TEXT("RmaxOut Sending: %u"), Index));
+			FRivermaxTracingUtils::RmaxOutFrameReadyTraceEvents.Add(Index, FString::Printf(TEXT("RmaxOut FrameReady: %u"), Index));
+			FRivermaxTracingUtils::RmaxInStartingFrameTraceEvents.Add(Index, FString::Printf(TEXT("RmaxIn Starting: %u"), Index));
+			FRivermaxTracingUtils::RmaxInReceivedFrameTraceEvents.Add(Index, FString::Printf(TEXT("RmaxIn Received: %u"), Index));
+			FRivermaxTracingUtils::RmaxInSelectedFrameTraceEvents.Add(Index, FString::Printf(TEXT("RmaxIn Selected: %u"), Index));
+		}
+	}
+
 }
 
 #undef LOCTEXT_NAMESPACE // RivermaxManager
diff --git a/Engine/Plugins/VirtualProduction/Rivermax/RivermaxCore/Source/RivermaxCore/Private/RivermaxManager.h b/Engine/Plugins/VirtualProduction/Rivermax/RivermaxCore/Source/RivermaxCore/Private/RivermaxManager.h
--- a/Engine/Plugins/VirtualProduction/Rivermax/RivermaxCore/Source/RivermaxCore/Private/RivermaxManager.h
+++ b/Engine/Plugins/VirtualProduction/Rivermax/RivermaxCore/Source/RivermaxCore/Private/RivermaxManager.h
@@ -53,6 +53,9 @@ namespace UE::RivermaxCore::Private
 
 		/** Verify if all prerequesites to go through Rivermax initialization have been met */
 		bool VerifyPrerequesites();
+
+		/** Initializes maps of strings used in various trace events to avoid allocating them continuously */
+		void InitializeTraceMarkupStrings();
 	private:
 
 		/** True when manager has been initialized */
diff --git a/Engine/Plugins/VirtualProduction/Rivermax/RivermaxCore/Source/RivermaxCore/Private/RivermaxOutputFrame.cpp b/Engine/Plugins/VirtualProduction/Rivermax/RivermaxCore/Source/RivermaxCore/Private/RivermaxOutputFrame.cpp
--- a/Engine/Plugins/VirtualProduction/Rivermax/RivermaxCore/Source/RivermaxCore/Private/RivermaxOutputFrame.cpp
+++ b/Engine/Plugins/VirtualProduction/Rivermax/RivermaxCore/Source/RivermaxCore/Private/RivermaxOutputFrame.cpp
@@ -28,6 +28,7 @@ namespace UE::RivermaxCore::Private
 	{
 		FrameIdentifier = InvalidIdentifier;
 		bIsVideoBufferReady = false;
+		bCaughtTimingIssue = false;
 		ReadyTimestamp = 0;
 		Clear();
 	}
diff --git a/Engine/Plugins/VirtualProduction/Rivermax/RivermaxCore/Source/RivermaxCore/Private/RivermaxOutputFrame.h b/Engine/Plugins/VirtualProduction/Rivermax/RivermaxCore/Source/RivermaxCore/Private/RivermaxOutputFrame.h
--- a/Engine/Plugins/VirtualProduction/Rivermax/RivermaxCore/Source/RivermaxCore/Private/RivermaxOutputFrame.h
+++ b/Engine/Plugins/VirtualProduction/Rivermax/RivermaxCore/Source/RivermaxCore/Private/RivermaxOutputFrame.h
@@ -53,8 +53,8 @@ namespace UE::RivermaxCore::Private
 		uint32 ChunkNumber = 0;
 		uint32 BytesSent = 0;
 
-		/** Timestamp of this frame  */
-		double TimestampTicks = 0.0;
+		/** Timestamp of this frame used for RTP headers  */
+		uint32 MediaTimestamp = 0;
 
 		/** Payload (data) pointer retrieved from Rivermax for the next chunk */
 		void* PayloadPtr = nullptr;
@@ -73,6 +73,9 @@ namespace UE::RivermaxCore::Private
 
 		/** Time at which this frame was made available to be sent */
 		uint64 ReadyTimestamp = 0;
+
+		/** Whether timing issues were detected while sending frame out. If yes, we skip the next frame boundary */
+		bool bCaughtTimingIssue = false;
 	};
 }
 
diff --git a/Engine/Plugins/VirtualProduction/Rivermax/RivermaxCore/Source/RivermaxCore/Private/RivermaxOutputStream.cpp b/Engine/Plugins/VirtualProduction/Rivermax/RivermaxCore/Source/RivermaxCore/Private/RivermaxOutputStream.cpp
--- a/Engine/Plugins/VirtualProduction/Rivermax/RivermaxCore/Source/RivermaxCore/Private/RivermaxOutputStream.cpp
+++ b/Engine/Plugins/VirtualProduction/Rivermax/RivermaxCore/Source/RivermaxCore/Private/RivermaxOutputStream.cpp
@@ -12,9 +12,11 @@
 #include "RivermaxFrameManager.h"
 #include "RivermaxLog.h"
 #include "RivermaxPTPUtils.h"
+#include "RivermaxTracingUtils.h"
 #include "RivermaxTypes.h"
 #include "RivermaxUtils.h"
 
+
 namespace UE::RivermaxCore::Private
 {
 	static TAutoConsoleVariable<int32> CVarRivermaxWakeupOffset(
@@ -59,7 +61,7 @@ namespace UE::RivermaxCore::Private
 		ECVF_Default);
 
 	static TAutoConsoleVariable<int32> CVarRivermaxOutputSkipSchedulingCutOffTime(
-		TEXT("Rivermax.Output.Scheduling.SkipCutoff"), 100,
+		TEXT("Rivermax.Output.Scheduling.SkipCutoff"), 50,
 		TEXT("Required time in microseconds from scheduling time to avoid skipping an interval."),
 		ECVF_Default);
 
@@ -68,14 +70,16 @@ namespace UE::RivermaxCore::Private
 		TEXT("Required time in nanoseconds from scheduling time before we clamp to do it immediately."),
 		ECVF_Default);
 
-	static TAutoConsoleVariable<int32> CVarRivermaxOutputFrameSliceCount(
-		TEXT("Rivermax.Output.FrameSliceCount"), 20,
-		TEXT("Number of memcopy done per frame when using intermediate buffer. As frame gets bigger, we can't do a single memcopy or timings will be broken."),
+	static TAutoConsoleVariable<int32> CVarRivermaxOutputMaxFrameMemorySliceCount(
+		TEXT("Rivermax.Output.FrameSliceCount"), 30,
+		TEXT("Max number of memcopies done per frame when using intermediate buffer. As frame gets bigger, we can't do a single memcopy or timings will be broken. Can be smaller in order to fit inside chunk count."),
 		ECVF_Default);
 
-	static TAutoConsoleVariable<int32> CVarRivermaxOutputEnableIntermediateBuffer(
-		TEXT("Rivermax.Output.EnableIntermediateBuffer"), 1,
-		TEXT("Uses an intermediate buffer used by Rivermax when sending data out. During scheduling, captured frame data will be copied over intermediate buffer."),
+	static TAutoConsoleVariable<bool> CVarRivermaxOutputEnableIntermediateBuffer(
+		TEXT("Rivermax.Output.Alignment.EnableIntermediateBuffer"), true,
+		TEXT("Uses an intermediate buffer used by Rivermax when sending data out.\n")
+		TEXT("During scheduling, captured frame data will be copied over intermediate buffer.\n")
+		TEXT("Only applies to alignment points scheduling mode."),
 		ECVF_Default);
 
 	static TAutoConsoleVariable<int32> CVarRivermaxOutputUseSingleMemblock(
@@ -83,6 +87,31 @@ namespace UE::RivermaxCore::Private
 		TEXT("Configures Rivermax stream to use a single memblock potentially improving SDK performance."),
 		ECVF_Default);
 
+	static TAutoConsoleVariable<int32> CVarRivermaxOutputRandomDelay(
+		TEXT("Rivermax.Output.TriggerRandomDelay"), 0,
+		TEXT("Will cause a delay of variable amount of time when next frame is sent."),
+		ECVF_Default);
+
+	static TAutoConsoleVariable<int32> CVarRivermaxOutputEnableTimingProtection(
+		TEXT("Rivermax.Output.Scheduling.EnableTimingProtection"), 1,
+		TEXT("Whether timing verification is done on commit to avoid misalignment. Next frame interval is skipped if it happens."),
+		ECVF_Default);
+
+	static TAutoConsoleVariable<int32> CVarRivermaxOutputMemcopyChunkSpacing(
+		TEXT("Rivermax.Output.Scheduling.MemcopyChunkSpacing"), 10,
+		TEXT("Number of chunks between each memcopy to help with timing for different frame format."),
+		ECVF_Default);
+
+	static TAutoConsoleVariable<int32> CVarRivermaxOutputEnableTimeCriticalThread(
+		TEXT("Rivermax.Output.EnableTimeCriticalThread"), 0,
+		TEXT("Whether to set output thread as time critical."),
+		ECVF_Default);
+
+	static TAutoConsoleVariable<bool> CVarRivermaxOutputPrefillRTPHeaders(
+		TEXT("Rivermax.Output.PrefillRTPHeaders"), true,
+		TEXT("Optimization used to prefill every RTP headers with known data."),
+		ECVF_Default);
+
 	bool FindPayloadSize(const FRivermaxOutputStreamOptions& InOptions, uint32 InBytesPerLine, const FVideoFormatInfo& FormatInfo, uint16& OutPayloadSize)
 	{
 		using namespace UE::RivermaxCore::Private::Utils;
@@ -327,6 +356,105 @@ namespace UE::RivermaxCore::Private
 		uint64 TimeCovered = 0;
 	};
 
+	struct FRTPHeaderPrefiller
+	{
+		FRTPHeaderPrefiller(FRivermaxOutputStream& InRmaxOutputStream)
+			: Stream(InRmaxOutputStream)
+		{
+			RunningSRDOffsetPerFrame.SetNumZeroed(Stream.Options.NumberOfBuffers);
+			RunningLineNumberPerFrame.SetNumZeroed(Stream.Options.NumberOfBuffers);
+		}
+
+		void Update(uint32 PacketIndex)
+		{
+			using namespace UE::RivermaxCore::Private::Utils;
+
+			const uint32 Width = Stream.Options.AlignedResolution.X;
+			const uint32 Height = Stream.Options.AlignedResolution.Y;
+			const int32 PacketCount = Stream.StreamMemory.ChunksPerFrameField * Stream.StreamMemory.PacketsPerChunk;
+			uint32 BufferIndex = 0;
+			for (int32 MemblockIndex = 0; MemblockIndex < Stream.StreamMemory.RTPHeaders.Num(); ++MemblockIndex)
+			{
+				for (uint32 FrameInBlockIndex = 0; FrameInBlockIndex < Stream.StreamMemory.FramesFieldPerMemoryBlock; ++FrameInBlockIndex)
+				{
+					uint16& SRDOffset = RunningSRDOffsetPerFrame[BufferIndex];
+					uint32& LineNumber = RunningLineNumberPerFrame[BufferIndex];
+					const uint32 HeaderIndex = PacketIndex + (FrameInBlockIndex * PacketCount);
+					if (ensure(Stream.StreamMemory.RTPHeaders[MemblockIndex].IsValidIndex(HeaderIndex)))
+					{
+						FRawRTPHeader& CurrentHeader = Stream.StreamMemory.RTPHeaders[MemblockIndex][HeaderIndex];
+						CurrentHeader = {};
+						CurrentHeader.Version = 2;
+						CurrentHeader.PaddingBit = 0;
+						CurrentHeader.ExtensionBit = 0;
+						CurrentHeader.PayloadType = 96; //Payload type should probably be inferred from SDP
+
+						//2110 specific header
+						CurrentHeader.SynchronizationSource = ByteSwap((uint32)0x0eb51dbd);  // Should Unreal has its own synch source ID
+
+						{
+							if (PacketIndex + 1 == Stream.StreamMemory.PacketsPerFrame)
+							{
+								CurrentHeader.MarkerBit = 1; // last packet in frame (Marker)
+							}
+
+							// Verify if payload size exceeds line 
+							const uint32 CurrentPayloadSize = Stream.StreamMemory.PayloadSizes[PacketIndex];
+
+							const uint32 LineSizeOffset = ((SRDOffset / Stream.FormatInfo.PixelGroupCoverage) * Stream.FormatInfo.PixelGroupSize);
+							const uint32 LineSize = ((Width / Stream.FormatInfo.PixelGroupCoverage) * Stream.FormatInfo.PixelGroupSize);
+
+							const uint16 SRD1Length = FMath::Min(LineSize - LineSizeOffset, CurrentPayloadSize);
+							const uint16 SRD1PixelCount = SRD1Length / Stream.FormatInfo.PixelGroupSize * Stream.FormatInfo.PixelGroupCoverage;
+							uint16 SRD2Length = SRD1Length < CurrentPayloadSize ? CurrentPayloadSize - SRD1Length : 0;
+							if (SRD2Length && LineNumber == (Height - 1))
+							{
+								SRD2Length = 0;
+							}
+
+							CurrentHeader.SRD1Length = ByteSwap(SRD1Length);
+							CurrentHeader.SetSrd1RowNumber(LineNumber); //todo divide by 2 if interlaced
+							CurrentHeader.FieldIdentification1 = 0; //todo when fields are sent for interlace
+							CurrentHeader.SetSrd1Offset(SRDOffset);
+
+							SRDOffset += SRD1PixelCount;
+							if (SRDOffset >= Width)
+							{
+								SRDOffset = 0;
+								++LineNumber;
+							}
+
+							if (SRD2Length > 0)
+							{
+								CurrentHeader.SRD2Length = ByteSwap(SRD2Length);
+
+								CurrentHeader.ContinuationBit1 = 1;
+								CurrentHeader.FieldIdentification2 = 0;
+								CurrentHeader.SetSrd2RowNumber(LineNumber);
+								CurrentHeader.SetSrd2Offset(SRDOffset);
+
+								const uint16 SRD2PixelCount = SRD2Length / Stream.FormatInfo.PixelGroupSize * Stream.FormatInfo.PixelGroupCoverage;
+								SRDOffset += SRD2PixelCount;
+								if (SRDOffset >= Width)
+								{
+									SRDOffset = 0;
+									++LineNumber;
+								}
+							}
+						}
+					}
+
+					++BufferIndex;
+				}
+			}
+		}
+
+	private:
+		TArray<uint16> RunningSRDOffsetPerFrame;
+		TArray<uint32> RunningLineNumberPerFrame;
+		FRivermaxOutputStream& Stream;
+	};
+
 	FRivermaxOutputStream::~FRivermaxOutputStream()
 	{
 		Uninitialize();
@@ -347,6 +475,11 @@ namespace UE::RivermaxCore::Private
 		Listener = &InListener;
 		FormatInfo = FStandardVideoFormat::GetVideoFormatInfo(Options.PixelFormat);
 
+		if (IConsoleVariable* CvarDelay = CVarRivermaxOutputRandomDelay.AsVariable())
+		{
+			CvarDelay->OnChangedDelegate().AddRaw(this, &FRivermaxOutputStream::OnCVarRandomDelayChanged);
+		}
+
 		const uint32 FrameSize = GetStride() * Options.AlignedResolution.Y;
 		if (FrameSize == 0)
 		{
@@ -426,6 +559,8 @@ namespace UE::RivermaxCore::Private
 						
 						UE_LOG(LogRivermax, Display, TEXT("%s"), *FString(StreamDescription));
 
+						UE_LOG(LogRivermax, Verbose, TEXT("Created stream using SDP:\n%S"), SDPDescription.GetData());
+
 						bIsActive = true;
 						RivermaxThread.Reset(FRunnableThread::Create(this, TEXT("Rmax OutputStream Thread"), 128 * 1024, TPri_TimeCritical, FPlatformAffinity::GetPoolThreadMask()));
 					}
@@ -464,6 +599,11 @@ namespace UE::RivermaxCore::Private
 			
 			UE_LOG(LogRivermax, Log, TEXT("Rivermax Output stream has shutdown"));
 		}
+
+		if (IConsoleVariable* CvarDelay = CVarRivermaxOutputRandomDelay.AsVariable())
+		{
+			CvarDelay->OnChangedDelegate().RemoveAll(this);
+		}
 	}
 
 	bool FRivermaxOutputStream::PushVideoFrame(const FRivermaxOutputVideoFrameInfo& NewFrame)
@@ -484,88 +624,128 @@ namespace UE::RivermaxCore::Private
 	{
 		using namespace UE::RivermaxCore::Private::Utils;
 
-		// Not the best path but it seems to work without tearing for now
 		// Wait for the next time a frame should be sent (based on frame interval)
-		// If a frame had been sent, make it available again
-		// This is to avoid making it available right after scheduling it. It's not sent yet and we start overwriting it
-		// Get next available frame that was rendered (Wait if none are)
+		// if interm buffer is used (alignment points) and a frame is ready before frame interval
+		//		Start copying data into next memory block from the intermediate buffer
+		//		At frame interval:
+		//			Release last sent frame if any
+		//			Make next frame the one being sent
+		// Otherwise
+		//		FrameCreation:
+		//			Release last sent frame if any
+		//			Wait for a new frame to be available
+		//		Alignment points:
+		//			Release last sent frame if any
+		// 
 		// Send frame
+		//		Get next chunk
+		//		Continue copy to intermediate buffer if required
+		//		Fill dynamic data for RTP headers of next chunk
+		//		Commit next chunk
+		// 
 		// Restart
 		{
+			bool bCanEarlyCopy = false;
 			{
-				TRACE_CPUPROFILER_EVENT_SCOPE(FRivermaxOutputStream::Wait);
-				WaitForNextRound();
+				TRACE_CPUPROFILER_EVENT_SCOPE(RmaxOut::Wait);
+				bCanEarlyCopy = WaitForNextRound();
 			}
 
-			// Verify if we can continue sending the next frame. We might be too late if our thread wasn't woken up in time
-			if (CanProceedWithScheduling())
+			if(bIsActive && bCanEarlyCopy)
 			{
+				TRACE_CPUPROFILER_EVENT_SCOPE(RmaxOut::PreprocessNextFrame);
+				PreprocessNextFrame();
+			}
+			else
+			{
+				TRACE_CPUPROFILER_EVENT_SCOPE(RmaxOut::PrepareNextFrame);
+				PrepareNextFrame();
+			}
 
-				Stats.LastFrameChunkRetries = 0;
+			// At this point, if there is no frame to send, move on to wait for next round
+			if (CurrentFrame.IsValid() && bIsActive)
+			{
+				SendFrame();
+				
+				// If frame that was just sent failed timing requirements, we have to tell rivermax to skip 0 chunks 
+				// in order to reset internal states. Otherwise, scheduling time / Tro isn't respected next time we schedule.
+				if (CurrentFrame->bCaughtTimingIssue)
 				{
-					TRACE_CPUPROFILER_EVENT_SCOPE(FRivermaxOutputStream::PrepareNextFrame);
-					PrepareNextFrame();
+					++Stats.TimingIssueCount;
+					
+					constexpr uint64 ChunksToSkip = 0;
+					SkipChunks(ChunksToSkip);
 				}
+			}
 
-				// At this point, if there is no frame to send, move on to wait for next round
-				if (CurrentFrame.IsValid() && bIsActive)
+			Stats.TotalChunkRetries += Stats.LastFrameChunkRetries;
+			StreamData.bHasFrameFirstChunkBeenFetched = false;
+			Stats.LastFrameChunkRetries = 0;
+		}
+	}
+
+	void FRivermaxOutputStream::PreprocessNextFrame()
+	{
+		checkSlow(Options.AlignmentMode == ERivermaxAlignmentMode::AlignmentPoint);
+
+		TSharedPtr<FRivermaxOutputFrame> NextFrameToSend = FrameManager->GetReadyFrame();
+		if (ensure(NextFrameToSend))
+		{
+			InitializeNextFrame(NextFrameToSend);
+
+			// Now that we have the next frame, we can start copying data into it.
+			// We can't get chunks since commit will only commit the chunks returned by last call to get next chunk. 
+			// So, we calculate the next data and header pointer based on the current frame. 
+			const uint64 CurrentRmaxTimeNanosec = RivermaxModule->GetRivermaxManager()->GetTime();
+			const double CurrentPlatformTime = FPlatformTime::Seconds();
+			const double TimeLeftSec = double(StreamData.NextAlignmentPointNanosec - CurrentRmaxTimeNanosec) / 1E9;
+			const double TargetPlatformTimeSec = CurrentPlatformTime + TimeLeftSec;
+			if (CurrentRmaxTimeNanosec < StreamData.NextAlignmentPointNanosec)
+			{
 				{
-					StreamData.LastSendStartTimeNanoSec = RivermaxModule->GetRivermaxManager()->GetTime();
+					TRACE_CPUPROFILER_EVENT_SCOPE(RmaxOut::CopyFrame);
 
-					// Add markup when we start pushing out a frame with its timestamp to track it across network
-					if (StreamData.bHasFrameFirstChunkBeenFetched == false)
+					bool bHasDataToCopy = true;
+					while (FPlatformTime::Seconds() < TargetPlatformTimeSec && bHasDataToCopy && bIsActive)
 					{
-						const uint32 MediaTimestamp = GetTimestampFromTime(StreamData.NextScheduleTimeNanosec, MediaClockSampleRate);
-						const FString TraceName = FString::Format(TEXT("RmaxOutput::StartSending {0}|{1}"), { MediaTimestamp, CurrentFrame->FrameIdentifier });
-						TRACE_CPUPROFILER_EVENT_SCOPE_TEXT(*TraceName);
-						UE_LOG(LogRivermax, Verbose, TEXT("RmaxTX frame number %u with timestamp %u."), CurrentFrame->FrameIdentifier, MediaTimestamp);
+						bHasDataToCopy = CopyFrameData(NextFrameToSend, reinterpret_cast<uint8*>(StreamMemory.BufferAddresses[StreamData.ExpectedFrameIndex]));
+						
 					}
+				}
 
-					const FString TraceName = FString::Format(TEXT("RmaxOutput::SendFrame {0}"), { CurrentFrame->FrameIndex });
-					TRACE_CPUPROFILER_EVENT_SCOPE_TEXT(*TraceName);
-					TRACE_CPUPROFILER_EVENT_SCOPE_TEXT(*FString::Printf(TEXT("MediaCapturePipe: %u"), CurrentFrame->FrameIdentifier));
+				const double PostCopyTimeLeftSec = TargetPlatformTimeSec - FPlatformTime::Seconds();
+				if (PostCopyTimeLeftSec > 0)
+				{
+					TRACE_CPUPROFILER_EVENT_SCOPE(RmaxOut::Waiting);
+					static constexpr float YieldTimeSec = 2.0f / 1000;
 
-					do 
+					while (FPlatformTime::Seconds() < TargetPlatformTimeSec && bIsActive)
 					{
-						if (bIsActive)
-						{
-							GetNextChunk();
-						}
-
-						if (bIsActive && CachedCVars.bUseIntermediateBuffer)
-						{
-							CopyFrameData();
-						}
-
-						if (bIsActive)
-						{
-							SetupRTPHeaders();
-						}
-
-						if (bIsActive)
-						{
-							CommitNextChunks();
-						}
-
-						//Update frame progress
-						if (bIsActive)
-						{
-							Stats.TotalStrides += StreamMemory.PacketsPerChunk;
-							++CurrentFrame->ChunkNumber;
-						}
-					} while (CurrentFrame->ChunkNumber < StreamMemory.ChunksPerFrameField && bIsActive);
-					
-					Stats.FramesSentCounter++;
-					StreamData.RTPHeaderFrameIndex = (StreamData.RTPHeaderFrameIndex + 1) % Options.NumberOfBuffers; 
+						const double TimeLeft = TargetPlatformTimeSec - FPlatformTime::Seconds();
+						const double SleepTime = TimeLeft > YieldTimeSec ? TimeLeft - YieldTimeSec : 0.0;
+						FPlatformProcess::SleepNoStats(SleepTime);
+					}
 				}
 			}
-			else
+
 			{
-				Stats.LastFrameChunkRetries = 0;
-			}
+				TRACE_CPUPROFILER_EVENT_SCOPE(RmaxOut::WrappingUp);
+				if (CurrentFrame)
+				{
+					FrameManager->MarkAsSent(CurrentFrame);
+					CurrentFrame.Reset();
+				}
 
-			Stats.TotalChunkRetries += Stats.LastFrameChunkRetries;
-			StreamData.bHasFrameFirstChunkBeenFetched = false;
+				// Make the next frame to send the current one and update its state
+				CurrentFrame = MoveTemp(NextFrameToSend);
+				FrameManager->MarkAsSending(CurrentFrame);
+			}
+		}
+		else
+		{
+			UE_LOG(LogRivermax, Error, TEXT("Unexpected error, no frame was available."));
+			Listener->OnStreamError();
+			Stop();
 		}
 	}
 
@@ -630,64 +810,17 @@ namespace UE::RivermaxCore::Private
 
 		// Setup arrays with the right sizes so we can give pointers to rivermax
 		StreamMemory.RTPHeaders.SetNumZeroed(StreamMemory.MemoryBlockCount);
-		StreamMemory.PayloadSizes.SetNumUninitialized(StreamMemory.PacketsPerMemoryBlock);
-		StreamMemory.HeaderSizes.SetNumUninitialized(StreamMemory.PacketsPerMemoryBlock);
+		StreamMemory.PayloadSizes.SetNumZeroed(StreamMemory.PacketsPerMemoryBlock);
+		StreamMemory.HeaderSizes.SetNumZeroed(StreamMemory.PacketsPerMemoryBlock);
 		StreamMemory.HeaderStrideSize = sizeof(FRawRTPHeader);
 
-	
-		for (uint32 BufferIndex = 0; BufferIndex < StreamMemory.FramesFieldPerMemoryBlock; ++BufferIndex)
-		{
-			uint64 TotalSize = 0;
-			uint64 LineSize = 0;
-			for (int32 PayloadSizeIndex = 0; PayloadSizeIndex < RealPacketsPerFrame; ++PayloadSizeIndex)
-			{
-				uint32 HeaderSize = FRawRTPHeader::OneSRDSize;
-				uint32 ThisPayloadSize = StreamMemory.PayloadSize;
-				if (TotalSize < FrameSize)
-				{
-					if ((LineSize + StreamMemory.PayloadSize) == BytesPerLine)
-					{
-						LineSize = 0;
-					}
-					else if ((LineSize + StreamMemory.PayloadSize) > BytesPerLine)
-					{
-						HeaderSize = FRawRTPHeader::TwoSRDSize;
-						LineSize = StreamMemory.PayloadSize - (BytesPerLine - LineSize);
-						if (LineSize > BytesPerLine)
-						{
-							UE_LOG(LogRivermax, Warning, TEXT("Unsupported small resolution, %dx%d, needing more than 2 SRD to express"), Options.AlignedResolution.X, Options.AlignedResolution.Y);
-							return false;
-						}
-					}
-					else
-					{
-						// Keep track of line size offset to know when to use TwoSRDs
-						LineSize += StreamMemory.PayloadSize;
-					}
-
-					if ((TotalSize + StreamMemory.PayloadSize) > FrameSize)
-					{
-						HeaderSize = FRawRTPHeader::OneSRDSize;
-					}
-				}
-				else
-				{
-					// Extra header/payload required for the chunk alignment are set to 0. Nothing has to be sent out the wire.
-					HeaderSize = 0;
-					ThisPayloadSize = 0;
-				}
-
-				StreamMemory.HeaderSizes[PayloadSizeIndex + (BufferIndex* RealPacketsPerFrame)] = HeaderSize;
-				StreamMemory.PayloadSizes[PayloadSizeIndex + (BufferIndex * RealPacketsPerFrame)] = ThisPayloadSize;
-				TotalSize += ThisPayloadSize;
-			}
-		}
-
 		if (!SetupFrameManagement())
 		{
 			return false;
 		}
-		
+
+		FRTPHeaderPrefiller RTPFiller(*this);
+
 		StreamMemory.MemoryBlocks.SetNumZeroed(StreamMemory.MemoryBlockCount);
 		for (uint32 BlockIndex = 0; BlockIndex < StreamMemory.MemoryBlockCount; ++BlockIndex)
 		{
@@ -697,7 +830,7 @@ namespace UE::RivermaxCore::Private
 			Block.data_size_arr = StreamMemory.PayloadSizes.GetData();
 
 			// If intermediate buffer is used, we setup rmax memblock to use that address. Otherwise, we map it to our actual frame's address
-			if(CachedCVars.bUseIntermediateBuffer)
+			if (StreamMemory.bUseIntermediateBuffer)
 			{
 				Block.data_ptr = Allocator->GetFrameAddress(BlockIndex);
 			}
@@ -709,6 +842,77 @@ namespace UE::RivermaxCore::Private
 			StreamMemory.RTPHeaders[BlockIndex].SetNumZeroed(RealPacketsPerFrame * StreamMemory.FramesFieldPerMemoryBlock);
 			Block.app_hdr_ptr = &StreamMemory.RTPHeaders[BlockIndex][0];
 		}
+	
+		uint64 TotalSize = 0;
+		uint64 LineSize = 0;
+		for (int32 PayloadSizeIndex = 0; PayloadSizeIndex < RealPacketsPerFrame; ++PayloadSizeIndex)
+		{
+			uint32 HeaderSize = FRawRTPHeader::OneSRDSize;
+			uint32 ThisPayloadSize = StreamMemory.PayloadSize;
+			if (TotalSize < FrameSize)
+			{
+				if ((LineSize + StreamMemory.PayloadSize) == BytesPerLine)
+				{
+					LineSize = 0;
+				}
+				else if ((LineSize + StreamMemory.PayloadSize) > BytesPerLine)
+				{
+					HeaderSize = FRawRTPHeader::TwoSRDSize;
+					LineSize = StreamMemory.PayloadSize - (BytesPerLine - LineSize);
+					if (LineSize > BytesPerLine)
+					{
+						UE_LOG(LogRivermax, Warning, TEXT("Unsupported small resolution, %dx%d, needing more than 2 SRD to express"), Options.AlignedResolution.X, Options.AlignedResolution.Y);
+						return false;
+					}
+				}
+				else
+				{
+					// Keep track of line size offset to know when to use TwoSRDs
+					LineSize += StreamMemory.PayloadSize;
+				}
+
+				if ((TotalSize + StreamMemory.PayloadSize) > FrameSize)
+				{
+					HeaderSize = FRawRTPHeader::OneSRDSize;
+				}
+			}
+			else
+			{
+				// Extra header/payload required for the chunk alignment are set to 0. Nothing has to be sent out the wire.
+				HeaderSize = 0;
+				ThisPayloadSize = 0;
+			}
+
+			// All buffers are configured the same so compute header and payload sizes once and assigned to all impacted locations
+			for (uint32 BufferIndex = 0; BufferIndex < StreamMemory.FramesFieldPerMemoryBlock; ++BufferIndex)
+			{
+				StreamMemory.HeaderSizes[PayloadSizeIndex + (BufferIndex * RealPacketsPerFrame)] = HeaderSize;
+				StreamMemory.PayloadSizes[PayloadSizeIndex + (BufferIndex * RealPacketsPerFrame)] = ThisPayloadSize;
+			}
+
+			if(CachedCVars.bPrefillRTPHeaders && HeaderSize > 0)
+			{
+				RTPFiller.Update(PayloadSizeIndex);
+			}
+			
+			TotalSize += ThisPayloadSize;
+		}
+
+		// Verify memcopy config to make sure it works for current frame size / chunking
+		if (StreamMemory.bUseIntermediateBuffer)
+		{
+			StreamMemory.FrameMemorySliceCount = FMath::Clamp(CVarRivermaxOutputMaxFrameMemorySliceCount.GetValueOnAnyThread(), 1, 100);
+			StreamMemory.ChunkSpacingBetweenMemcopies = FMath::Clamp(CVarRivermaxOutputMemcopyChunkSpacing.GetValueOnAnyThread(), 1, 20);
+
+			const uint32 ChunkRequired = StreamMemory.ChunkSpacingBetweenMemcopies * StreamMemory.FrameMemorySliceCount;
+			if (ChunkRequired > 0 && ChunkRequired > StreamMemory.ChunksPerFrameField)
+			{
+				// Favor reducing number of memcopies. If required packet is smaller, chances are it's a small frame size
+				// so memcopies will be smaller.
+				const double Ratio = StreamMemory.ChunksPerFrameField / (double)ChunkRequired;
+				StreamMemory.FrameMemorySliceCount *= Ratio;
+			}
+		}
 
 		return true;
 	}
@@ -723,84 +927,86 @@ namespace UE::RivermaxCore::Private
 		NextFrame->HeaderPtr = nullptr;
 		NextFrame->FrameStartPtr = nullptr;
 		NextFrame->ScheduleTimeCopied = 0;
+		NextFrame->bCaughtTimingIssue = false;
+		NextFrame->Offset = 0;
 	}
 
 	void FRivermaxOutputStream::BuildRTPHeader(FRawRTPHeader& OutHeader) const
 	{
 		using namespace UE::RivermaxCore::Private::Utils;
 
-		OutHeader = {};
-		OutHeader.Version = 2;
-		OutHeader.PaddingBit = 0;
-		OutHeader.ExtensionBit = 0;
-		OutHeader.PayloadType = 96; //Payload type should probably be infered from SDP
-		OutHeader.SequenceNumber = ByteSwap((uint16)(StreamData.SequenceNumber & 0xFFFF));
-
-		// For now, in order to be able to use a framelocked input, we pipe frame number in the timestamp for a UE-UE interaction
-		// Follow up work to investigate adding this in RTP header
-		uint64 InputTime = StreamData.NextScheduleTimeNanosec;
-		if (Options.bDoFrameCounterTimestamping)
+		if (CachedCVars.bPrefillRTPHeaders)
 		{
-			InputTime = UE::RivermaxCore::GetAlignmentPointFromFrameNumber(CurrentFrame->FrameIdentifier, Options.FrameRate);
+			OutHeader.SequenceNumber = ByteSwap((uint16)(StreamData.SequenceNumber & 0xFFFF));
+			OutHeader.Timestamp = ByteSwap(CurrentFrame->MediaTimestamp);
+			OutHeader.ExtendedSequenceNumber = ByteSwap((uint16)((StreamData.SequenceNumber >> 16) & 0xFFFF));
 		}
-
-		const uint32 MediaTimestamp = GetTimestampFromTime(InputTime, MediaClockSampleRate);
-		OutHeader.Timestamp = ByteSwap(MediaTimestamp);
-
-		//2110 specific header
-		OutHeader.SynchronizationSource = ByteSwap((uint32)0x0eb51dbd);  // Should Unreal has its own synch source ID
-
-		if (StreamType == ERivermaxStreamType::VIDEO_2110_20_STREAM)
+		else
 		{
-			if (CurrentFrame->PacketCounter + 1 == StreamMemory.PacketsPerFrame)
-			{
-				OutHeader.MarkerBit = 1; // last packet in frame (Marker)
-			}
+			OutHeader = {};
+			OutHeader.Version = 2;
+			OutHeader.PaddingBit = 0;
+			OutHeader.ExtensionBit = 0;
+			OutHeader.PayloadType = 96; //Payload type should probably be infered from SDP
 
-			OutHeader.ExtendedSequenceNumber = ByteSwap((uint16)((StreamData.SequenceNumber >> 16) & 0xFFFF));
-
-			// Verify if payload size exceeds line 
-			const uint32 CurrentPayloadSize = StreamMemory.PayloadSizes[CurrentFrame->PacketCounter];
+			OutHeader.SequenceNumber = ByteSwap((uint16)(StreamData.SequenceNumber & 0xFFFF));
+			OutHeader.Timestamp = ByteSwap(CurrentFrame->MediaTimestamp);
 
-			const uint32 LineSizeOffset = ((CurrentFrame->SRDOffset / FormatInfo.PixelGroupCoverage) * FormatInfo.PixelGroupSize);
-			const uint32 LineSize = ((Options.AlignedResolution.X / FormatInfo.PixelGroupCoverage) * FormatInfo.PixelGroupSize);
+			//2110 specific header
+			OutHeader.SynchronizationSource = ByteSwap((uint32)0x0eb51dbd);  // Should Unreal has its own synch source ID
 
-			const uint16 SRD1Length = FMath::Min(LineSize - LineSizeOffset, CurrentPayloadSize);
-			const uint16 SRD1PixelCount = SRD1Length / FormatInfo.PixelGroupSize * FormatInfo.PixelGroupCoverage;
-			uint16 SRD2Length = SRD1Length < CurrentPayloadSize ? CurrentPayloadSize - SRD1Length : 0;
-			if (SRD2Length && CurrentFrame->LineNumber == ((uint32)Options.AlignedResolution.Y-1))
+			if (StreamType == ERivermaxStreamType::VIDEO_2110_20_STREAM)
 			{
-				SRD2Length = 0;
-			}
+				if (CurrentFrame->PacketCounter + 1 == StreamMemory.PacketsPerFrame)
+				{
+					OutHeader.MarkerBit = 1; // last packet in frame (Marker)
+				}
 
-			OutHeader.SRD1Length = ByteSwap(SRD1Length);	
-			OutHeader.SetSrd1RowNumber(CurrentFrame->LineNumber); //todo divide by 2 if interlaced
-			OutHeader.FieldIdentification1 = 0; //todo when fields are sent for interlace
-			OutHeader.SetSrd1Offset(CurrentFrame->SRDOffset); 
+				OutHeader.ExtendedSequenceNumber = ByteSwap((uint16)((StreamData.SequenceNumber >> 16) & 0xFFFF));
 
-			CurrentFrame->SRDOffset += SRD1PixelCount;
-			if (CurrentFrame->SRDOffset >= Options.AlignedResolution.X)
-			{
-				CurrentFrame->SRDOffset = 0;
-				++CurrentFrame->LineNumber;
-			}
+				// Verify if payload size exceeds line 
+				const uint32 CurrentPayloadSize = StreamMemory.PayloadSizes[CurrentFrame->PacketCounter];
 
-			if (SRD2Length > 0)
-			{
-				OutHeader.SRD2Length = ByteSwap(SRD2Length);
+				const uint32 LineSizeOffset = ((CurrentFrame->SRDOffset / FormatInfo.PixelGroupCoverage) * FormatInfo.PixelGroupSize);
+				const uint32 LineSize = ((Options.AlignedResolution.X / FormatInfo.PixelGroupCoverage) * FormatInfo.PixelGroupSize);
+
+				const uint16 SRD1Length = FMath::Min(LineSize - LineSizeOffset, CurrentPayloadSize);
+				const uint16 SRD1PixelCount = SRD1Length / FormatInfo.PixelGroupSize * FormatInfo.PixelGroupCoverage;
+				uint16 SRD2Length = SRD1Length < CurrentPayloadSize ? CurrentPayloadSize - SRD1Length : 0;
+				if (SRD2Length && CurrentFrame->LineNumber == ((uint32)Options.AlignedResolution.Y - 1))
+				{
+					SRD2Length = 0;
+				}
 
-				OutHeader.ContinuationBit1 = 1;
-				OutHeader.FieldIdentification2 = 0;
-				OutHeader.SetSrd2RowNumber(CurrentFrame->LineNumber);
-				OutHeader.SetSrd2Offset(CurrentFrame->SRDOffset);
+				OutHeader.SRD1Length = ByteSwap(SRD1Length);
+				OutHeader.SetSrd1RowNumber(CurrentFrame->LineNumber); //todo divide by 2 if interlaced
+				OutHeader.FieldIdentification1 = 0; //todo when fields are sent for interlace
+				OutHeader.SetSrd1Offset(CurrentFrame->SRDOffset);
 
-				const uint16 SRD2PixelCount = SRD2Length / FormatInfo.PixelGroupSize * FormatInfo.PixelGroupCoverage;
-				CurrentFrame->SRDOffset += SRD2PixelCount;
+				CurrentFrame->SRDOffset += SRD1PixelCount;
 				if (CurrentFrame->SRDOffset >= Options.AlignedResolution.X)
 				{
 					CurrentFrame->SRDOffset = 0;
 					++CurrentFrame->LineNumber;
 				}
+
+				if (SRD2Length > 0)
+				{
+					OutHeader.SRD2Length = ByteSwap(SRD2Length);
+
+					OutHeader.ContinuationBit1 = 1;
+					OutHeader.FieldIdentification2 = 0;
+					OutHeader.SetSrd2RowNumber(CurrentFrame->LineNumber);
+					OutHeader.SetSrd2Offset(CurrentFrame->SRDOffset);
+
+					const uint16 SRD2PixelCount = SRD2Length / FormatInfo.PixelGroupSize * FormatInfo.PixelGroupCoverage;
+					CurrentFrame->SRDOffset += SRD2PixelCount;
+					if (CurrentFrame->SRDOffset >= Options.AlignedResolution.X)
+					{
+						CurrentFrame->SRDOffset = 0;
+						++CurrentFrame->LineNumber;
+					}
+				}
 			}
 		}
 	}
@@ -818,12 +1024,12 @@ namespace UE::RivermaxCore::Private
 			Status = RivermaxModule->GetRivermaxManager()->GetApi()->rmax_out_destroy_stream(StreamId);
 			if (RMAX_ERR_BUSY == Status) 
 			{
-				FPlatformProcess::SleepNoStats(SleepTimeSeconds);
+				FPlatformProcess::SleepNoStats(Utils::SleepTimeSeconds);
 			}
 		} while (Status == RMAX_ERR_BUSY);
 	}
 
-	void FRivermaxOutputStream::WaitForNextRound()
+	bool FRivermaxOutputStream::WaitForNextRound()
 	{
 		const uint64 CurrentTimeNanosec = RivermaxModule->GetRivermaxManager()->GetTime();
 		const double CurrentPlatformTime = FPlatformTime::Seconds();
@@ -858,51 +1064,83 @@ namespace UE::RivermaxCore::Private
 			WaitTimeNanosec = 0;
 		}
 
-		static constexpr float SleepThresholdSec = 5.0f * (1.0f / 1000.0f);
-		static constexpr float YieldTimeSec = 2.0f * (1.0f / 1000.0f);
+		static constexpr float SleepThresholdSec = 5.0f / 1000;
+		static constexpr float YieldTimeSec = 2.0f / 1000;
 		const double WaitTimeSec = FMath::Min(WaitTimeNanosec / 1E9, 1.0);
 		StreamData.LastSleepTimeNanoSec = WaitTimeNanosec;
 
-		// Sleep for the largest chunk of time
-		if (WaitTimeSec > SleepThresholdSec)
+		bool bIsFrameReady = false;
+		if (StreamMemory.bUseIntermediateBuffer)
+		{	
+			// When using intermediate buffer, we verify if next buffer is ready sooner than wake up time
+			// If a frame is ready already, we can move on. Otherwise, we wait for FrameReady signal with a
+			// wait timeout. In the case of a repeated frame, we will always timeout and we won't be able to 
+			// do an early copy.
+			TSharedPtr<FRivermaxOutputFrame> ReadyFrame = FrameManager->GetReadyFrame();
+			if (ReadyFrame)
+			{
+				bIsFrameReady = true;
+			}
+			else
+			{
+				const uint32 WaitMs = FMath::Floor((WaitTimeSec - YieldTimeSec) * 1000.);
+				do 
+				{
+					bIsFrameReady = FrameReadyToSendSignal->Wait(WaitMs);
+					ReadyFrame = FrameManager->GetReadyFrame();
+				} while (!ReadyFrame && bIsFrameReady && bIsActive);
+			}
+		}
+		else
 		{
-			FPlatformProcess::SleepNoStats(WaitTimeSec - YieldTimeSec);
+			// Sleep for the largest chunk of time
+			if (WaitTimeSec > SleepThresholdSec)
+			{
+				FPlatformProcess::SleepNoStats(WaitTimeSec - YieldTimeSec);
+			}
 		}
 
-		// Yield our thread time for the remainder of wait time. 
-		// Should we spin for a smaller time to be more precise?
-		// Should we use platform time instead of rivermax get PTP to avoid making calls to it?
-		constexpr bool bUsePlatformTimeToSpin = true;
-		if (bUsePlatformTimeToSpin == false)
+		if (!bIsFrameReady)
 		{
-			while (RivermaxModule->GetRivermaxManager()->GetTime() < WakeupTime)
+			// We are past the long sleep so no more early data access possible. Just yield until the wake up time.
 			{
-				FPlatformProcess::SleepNoStats(0.f);
+				// Use platform time instead of rivermax get PTP to avoid making calls to it. Haven't been profiled if it impacts
+				while (FPlatformTime::Seconds() < (CurrentPlatformTime + WaitTimeSec))
+				{
+					FPlatformProcess::SleepNoStats(0.f);
+				}
+			}
+
+			if (StreamData.bHasValidNextFrameNumber && CachedCVars.bShowOutputStats)
+			{
+				const uint64 AfterSleepTimeNanosec = RivermaxModule->GetRivermaxManager()->GetTime();
+				const uint64 RealWaitNs = AfterSleepTimeNanosec - CurrentTimeNanosec;
+				const uint64 OvershootSleep = AfterSleepTimeNanosec > StreamData.NextAlignmentPointNanosec ? AfterSleepTimeNanosec - StreamData.NextAlignmentPointNanosec : 0;
+				const double OvershootSleepSec = OvershootSleep / 1e9;
+
+				UE_LOG(LogRivermax, Verbose, TEXT("CurrentTime %llu. OvershootSleep: %0.9f. ExpectedWait: %0.9f. RealWait: %0.9f, Scheduling at %llu. NextAlign %llu. ")
+					, CurrentTimeNanosec
+					, OvershootSleepSec
+					, (double)WaitTimeNanosec / 1E9
+					, (double)RealWaitNs / 1E9
+					, StreamData.NextScheduleTimeNanosec
+					, StreamData.NextAlignmentPointNanosec);
 			}
 		}
 		else
 		{
-			while (FPlatformTime::Seconds() < (CurrentPlatformTime + WaitTimeSec))
+			if (StreamData.bHasValidNextFrameNumber && CachedCVars.bShowOutputStats)
 			{
-				FPlatformProcess::SleepNoStats(0.f);
+				const uint64 AfterSleepTimeNanosec = RivermaxModule->GetRivermaxManager()->GetTime();
+				UE_LOG(LogRivermax, Verbose, TEXT("Early data available. CurrentTime %llu. Scheduling at %llu. NextAlign %llu. ")
+					, AfterSleepTimeNanosec
+					, StreamData.NextScheduleTimeNanosec
+					, StreamData.NextAlignmentPointNanosec
+					, (StreamData.NextScheduleTimeNanosec - AfterSleepTimeNanosec));
 			}
 		}
 
-		if (StreamData.bHasValidNextFrameNumber)
-		{
-			const uint64 AfterSleepTimeNanosec = RivermaxModule->GetRivermaxManager()->GetTime();
-			const uint64 RealWaitNs = AfterSleepTimeNanosec - CurrentTimeNanosec;
-			const uint64 OvershootSleep = AfterSleepTimeNanosec > StreamData.NextAlignmentPointNanosec ? AfterSleepTimeNanosec - StreamData.NextAlignmentPointNanosec : 0;
-			const double OvershootSleepSec = OvershootSleep / 1e9;
-
-			UE_LOG(LogRivermax, Verbose, TEXT("CurrentTime %llu. OvershootSleep: %0.9f. ExpectedWait: %0.9f. RealWait: %0.9f, Scheduling at %llu. NextAlign %llu. ")
-			, CurrentTimeNanosec
-			, OvershootSleepSec
-			, (double)WaitTimeNanosec / 1E9
-			, (double)RealWaitNs / 1E9
-			, StreamData.NextScheduleTimeNanosec
-			, StreamData.NextAlignmentPointNanosec);
-		}
+		return bIsFrameReady;
 	}
 
 	void FRivermaxOutputStream::GetNextChunk()
@@ -920,15 +1158,10 @@ namespace UE::RivermaxCore::Private
 
 					// Stamp frame start in order to copy frame data sequentially as we query chunks
 					CurrentFrame->FrameStartPtr = CurrentFrame->PayloadPtr;
-					CurrentFrame->Offset = 0;
 
-					// If we aren't using intermediate buffer, we expect video buffer to point to the next chunk's payload pointer
-					if (CachedCVars.bUseIntermediateBuffer == false)
+					if (CurrentFrame->PayloadPtr != StreamMemory.BufferAddresses[StreamData.ExpectedFrameIndex])
 					{
-						if (CurrentFrame->VideoBuffer != CurrentFrame->PayloadPtr)
-						{
-							UE_LOG(LogRivermax, Warning, TEXT("Frame being sent (%d) doesn't match chunks being processed."), CurrentFrame->FrameIndex);
-						}
+						UE_LOG(LogRivermax, Warning, TEXT("Frame being sent (%d) doesn't match chunks being processed."), CurrentFrame->FrameIndex);
 					}
 				}
 
@@ -954,14 +1187,14 @@ namespace UE::RivermaxCore::Private
 		} while (Status != RMAX_OK && bIsActive);
 	}
 
-	void FRivermaxOutputStream::CopyFrameData()
+	bool FRivermaxOutputStream::CopyFrameData(const TSharedPtr<FRivermaxOutputFrame>& SourceFrame, uint8* DestinationBase)
 	{
 		// Make sure copy size doesn't go over frame size
 		const uint32 FrameSize = GetStride() * Options.AlignedResolution.Y;
 		if (ensure(FrameSize > 0))
 		{
-			const uint32 BlockSize = 1 + ((FrameSize - 1) / CachedCVars.FrameMemorySliceCount);
-			const uint32 MaxSize = FrameSize - CurrentFrame->Offset;
+			const uint32 BlockSize = 1 + ((FrameSize - 1) / StreamMemory.FrameMemorySliceCount);
+			const uint32 MaxSize = FrameSize - SourceFrame->Offset;
 			const uint32 CopySize = FMath::Min(BlockSize, MaxSize);
 
 			// Copy data until we have covered the whole frame. Last block might be smaller.
@@ -969,7 +1202,7 @@ namespace UE::RivermaxCore::Private
 			{
 				TRACE_CPUPROFILER_EVENT_SCOPE(CopyFrameData);
 				const uint32 PacketsToCopy = CopySize / StreamMemory.PayloadSize;
-				const uint32 PacketOffset = CurrentFrame->ChunkNumber * StreamMemory.PacketsPerChunk;
+				const uint32 PacketOffset = SourceFrame->ChunkNumber * StreamMemory.PacketsPerChunk;
 				const uint64 PacketDeltaTime = (StreamData.FrameFieldTimeIntervalNs - TransmitOffsetNanosec) / StreamMemory.PacketsPerFrame;
 
 				// Time covered in frame interval including this copy
@@ -980,17 +1213,19 @@ namespace UE::RivermaxCore::Private
 				CopyInfo->TimeCovered = TimeCopied;
 
 				FCopyArgs Args;
-				uint8* SourceStart = reinterpret_cast<uint8*>(CurrentFrame->VideoBuffer);
-				uint8* DestinationStart = reinterpret_cast<uint8*>(CurrentFrame->FrameStartPtr);
-				Args.SourceMemory = SourceStart + CurrentFrame->Offset;
-				Args.DestinationMemory = DestinationStart + CurrentFrame->Offset;
+				uint8* SourceStart = reinterpret_cast<uint8*>(SourceFrame->VideoBuffer);
+				uint8* DestinationStart = DestinationBase;
+				Args.SourceMemory = SourceStart + SourceFrame->Offset;
+				Args.DestinationMemory = DestinationStart + SourceFrame->Offset;
 				Args.SizeToCopy = CopySize;
 				Args.SideCar = MoveTemp(CopyInfo);
 
 				// Update memory offset for next copy
-				CurrentFrame->Offset += CopySize;
+				SourceFrame->Offset += CopySize;
 
 				Allocator->CopyData(Args);
+
+				return true;
 			}
 		}
 		else
@@ -999,6 +1234,8 @@ namespace UE::RivermaxCore::Private
 			Listener->OnStreamError();
 			Stop();
 		}
+
+		return false;
 	}
 
 	void FRivermaxOutputStream::SetupRTPHeaders()
@@ -1006,8 +1243,8 @@ namespace UE::RivermaxCore::Private
 		FRawRTPHeader* HeaderRawPtr = reinterpret_cast<FRawRTPHeader*>(CurrentFrame->HeaderPtr);
 		check(HeaderRawPtr);
 
-		const uint32 MemblockIndex = StreamData.RTPHeaderFrameIndex / StreamMemory.FramesFieldPerMemoryBlock;
-		const uint32 MemblockFrameIndex = StreamData.RTPHeaderFrameIndex % StreamMemory.FramesFieldPerMemoryBlock;
+		const uint32 MemblockIndex = StreamData.ExpectedFrameIndex / StreamMemory.FramesFieldPerMemoryBlock;
+		const uint32 MemblockFrameIndex = StreamData.ExpectedFrameIndex % StreamMemory.FramesFieldPerMemoryBlock;
 		const uint32 PacketsPerFrame = StreamMemory.ChunksPerFrameField * StreamMemory.PacketsPerChunk;
 		const uint32 BaseHeaderIndex = CurrentFrame->PacketCounter + (MemblockFrameIndex * PacketsPerFrame);
 
@@ -1029,11 +1266,11 @@ namespace UE::RivermaxCore::Private
 		rmax_status_t Status;
 		int32 ErrorCount = 0;
 		const uint64 CurrentTimeNanosec = RivermaxModule->GetRivermaxManager()->GetTime();
+		uint64 ScheduleTime = CurrentFrame->ChunkNumber == 0 ? StreamData.NextScheduleTimeNanosec : 0;
 
 		do
 		{
 			//Only first chunk gets scheduled with a timestamp. Following chunks are queued after it using 0
-			uint64 ScheduleTime = CurrentFrame->ChunkNumber == 0 ? StreamData.NextScheduleTimeNanosec : 0;
 			const rmax_commit_flags_t CommitFlags{};
 			if (ScheduleTime != 0)
 			{
@@ -1073,11 +1310,11 @@ namespace UE::RivermaxCore::Private
 
 		} while (Status != RMAX_OK && bIsActive);
 
-		if (bIsActive && CurrentFrame->ChunkNumber == 0)
+		if (bIsActive && CurrentFrame->ChunkNumber == 0 && CachedCVars.bShowOutputStats)
 		{
 			UE_LOG(LogRivermax, Verbose, TEXT("Committed frame [%u]. Scheduled for '%llu'. Aligned with '%llu'. Current time '%llu'. Was late: %d. Slack: %llu. Errorcount: %d")
 				, CurrentFrame->FrameIdentifier
-				, StreamData.NextScheduleTimeNanosec
+				, ScheduleTime
 				, StreamData.NextAlignmentPointNanosec
 				, CurrentTimeNanosec
 				, CurrentTimeNanosec >= StreamData.NextScheduleTimeNanosec ? 1 : 0
@@ -1093,6 +1330,11 @@ namespace UE::RivermaxCore::Private
 
 	uint32 FRivermaxOutputStream::Run()
 	{
+		if (CachedCVars.bEnableTimeCriticalThread)
+		{
+			::SetThreadPriority(GetCurrentThread(), THREAD_PRIORITY_TIME_CRITICAL);
+		}
+
 		// Initial wait for a frame to be produced
 		{
 			TRACE_CPUPROFILER_EVENT_SCOPE(Rmax::InitialWait);
@@ -1207,42 +1449,16 @@ namespace UE::RivermaxCore::Private
 			
 			// If intermediate buffer isn't used and frame has to be repeated, we use skip chunk method 
 			// which might cause timing errors caused by chunk managements issue 
-			if (!CachedCVars.bUseIntermediateBuffer)
+			if (!StreamMemory.bUseIntermediateBuffer)
 			{
 				// No frame to send, keep last one and restart its internal counters
 				UE_LOG(LogRivermax, Verbose, TEXT("No frame to send. Reusing last frame '%d' with Id %u"), CurrentFrame->FrameIndex, CurrentFrame->FrameIdentifier);
 
 				// Since we want to resend last frame, we need to fast forward chunk pointer to re-point to the one we just sent
-				rmax_status_t Status;
-				bool bHasAddedTrace = false;
-				do
-				{
-					checkSlow(CachedAPI);
-					Status = CachedAPI->rmax_out_skip_chunks(StreamId, StreamMemory.ChunksPerFrameField * (Options.NumberOfBuffers - 1));
-					if (Status != RMAX_OK)
-					{
-						if (Status == RMAX_ERR_NO_FREE_CHUNK)
-						{
-							// Wait until there are enough free chunk to be skipped
-							if (!bHasAddedTrace)
-							{
-								UE_LOG(LogRivermax, Warning, TEXT("No chunks ready to skip. Waiting"));
-								TRACE_CPUPROFILER_EVENT_SCOPE(PrepareNextFrame::NoFreeChunk);
-								bHasAddedTrace = true;
-							}
-						}
-						else
-						{
-							ensure(false);
-							UE_LOG(LogRivermax, Error, TEXT("Invalid error happened while trying to skip chunks. Status: %d."), Status);
-							Listener->OnStreamError();
-							Stop();
-						}
-					}
-				} while (Status != RMAX_OK && bIsActive);
-
+				SkipChunks(StreamMemory.ChunksPerFrameField * (Options.NumberOfBuffers - 1));
+				
 				// Keep RTP header in sync when skipping chunks
-				StreamData.RTPHeaderFrameIndex = CurrentFrame->FrameIndex;
+				StreamData.ExpectedFrameIndex = CurrentFrame->FrameIndex;
 			}
 		}
 	}
@@ -1330,13 +1546,13 @@ namespace UE::RivermaxCore::Private
 
 	void FRivermaxOutputStream::ShowStats()
 	{
-		if (CVarRivermaxOutputShowStats.GetValueOnAnyThread() != 0)
+		if (CachedCVars.bShowOutputStats)
 		{
 			const double CurrentTime = FPlatformTime::Seconds();
-			if (CurrentTime - LastStatsShownTimestamp > CVarRivermaxOutputShowStatsInterval.GetValueOnAnyThread())
+			if (CurrentTime - LastStatsShownTimestamp > CachedCVars.ShowOutputStatsIntervalSeconds)
 			{
 				LastStatsShownTimestamp = CurrentTime;
-				UE_LOG(LogRivermax, Log, TEXT("Stats: FrameSent: %u. CommitImmediate: %u. CommitRetries: %u. ChunkRetries: %u. ChunkSkippingRetries: %u"), Stats.FramesSentCounter, Stats.CommitImmediate, Stats.CommitRetries, Stats.TotalChunkRetries, Stats.ChunkSkippingRetries);
+				UE_LOG(LogRivermax, Log, TEXT("Stats: FrameSent: %llu. CommitImmediate: %u. CommitRetries: %u. ChunkRetries: %u. ChunkSkippingRetries: %u. Timing issues: %llu"), Stats.FramesSentCounter, Stats.CommitImmediate, Stats.CommitRetries, Stats.TotalChunkRetries, Stats.ChunkSkippingRetries, Stats.TimingIssueCount);
 			}
 		}
 	}
@@ -1382,8 +1598,9 @@ namespace UE::RivermaxCore::Private
 		bUseGPUDirect = FrameLocation == EFrameMemoryLocation::GPU;
 
 
+		// Only support intermediate buffer for alignment point method to avoid running into chunk issue when repeating a frame
 		const bool bHasAllocatedFrames = (FrameLocation != EFrameMemoryLocation::None);
-		if (bHasAllocatedFrames && CachedCVars.bUseIntermediateBuffer)
+		if (bHasAllocatedFrames && Options.AlignmentMode == ERivermaxAlignmentMode::AlignmentPoint && CVarRivermaxOutputEnableIntermediateBuffer.GetValueOnAnyThread())
 		{
 			// Allocate intermediate buffer in same memory space as frame memory.
 			FOnFrameDataCopiedDelegate OnDataCopiedDelegate = FOnFrameDataCopiedDelegate::CreateRaw(this, &FRivermaxOutputStream::OnMemoryChunksCopied);
@@ -1402,8 +1619,36 @@ namespace UE::RivermaxCore::Private
 			{
 				return false;
 			}
+
+			StreamMemory.bUseIntermediateBuffer = true;
 		}
 
+		// Cache buffer addresses used by rivermax in order to start copying into it early on
+		StreamMemory.BufferAddresses.Reserve(Options.NumberOfBuffers);
+		if (StreamMemory.bUseIntermediateBuffer)
+		{
+			check(Allocator);
+
+			for (uint32 MemblockIndex = 0; MemblockIndex < StreamMemory.MemoryBlockCount; ++MemblockIndex)
+			{
+				uint8* BaseAddress = static_cast<uint8*>(Allocator->GetFrameAddress(MemblockIndex));
+				for (uint32 FrameIndex = 0; FrameIndex < StreamMemory.FramesFieldPerMemoryBlock; ++FrameIndex)
+				{
+					uint8* FrameAddress = BaseAddress + (FrameIndex * FrameManagerArgs.FrameDesiredSize);
+					StreamMemory.BufferAddresses.Add(FrameAddress);
+				}
+			}
+		}
+		else
+		{
+			// When we don't use intermediate buffer, each frame has its own buffer addresses and we don't need to look at memblocks
+			for (int32 BufferIndex = 0; BufferIndex < Options.NumberOfBuffers; ++BufferIndex)
+			{
+				StreamMemory.BufferAddresses.Add(FrameManager->GetFrame(BufferIndex)->VideoBuffer);
+			}
+		}
+		check(StreamMemory.BufferAddresses.Num() == Options.NumberOfBuffers);
+
 		return FrameLocation != EFrameMemoryLocation::None;
 	}
 
@@ -1444,7 +1689,7 @@ namespace UE::RivermaxCore::Private
 			NextFrameNumber = CurrentFrameNumber + 1;
 		}
 		else
-		{
+		{	
 			// Case where we are back and frame number is the previous one. Depending on offsets, this could happen
 			if (CurrentFrameNumber == StreamData.NextAlignmentPointFrameNumber - 1)
 			{
@@ -1458,19 +1703,28 @@ namespace UE::RivermaxCore::Private
 				// We expect current frame number to be the one we scheduled for the last time or greater if something happened
 				if (CurrentFrameNumber >= StreamData.NextAlignmentPointFrameNumber)
 				{
-					// If current frame is greater than last scheduled, we missed an alignment point. Shouldn't happen with continuous thread independent of engine.
-					const uint64 DeltaFrames = CurrentFrameNumber - StreamData.NextAlignmentPointFrameNumber;
-					if (DeltaFrames >= 1)
+					// Verify if last frame had timing issues. If yes, we skip next interval
+					if (CurrentFrame && CurrentFrame->bCaughtTimingIssue)
 					{
-						UE_LOG(LogRivermax, Warning, TEXT("Output missed %llu frames."), DeltaFrames);
-						
-						// If we missed a sync point, this means that last scheduled frame might still be ongoing and 
-						// sending it might be crossing the frame boundary so we skip one entire frame to empty the queue.
 						NextFrameNumber = CurrentFrameNumber + 2;
+						UE_LOG(LogRivermax, Warning, TEXT("Timing issue detected during frame %llu. Skipping frame %llu to keep sync."), CurrentFrameNumber, CurrentFrameNumber + 1);
 					}
 					else
 					{
-						NextFrameNumber = CurrentFrameNumber + 1;
+						// If current frame is greater than last scheduled, we missed an alignment point.
+						const uint64 DeltaFrames = CurrentFrameNumber - StreamData.NextAlignmentPointFrameNumber;
+						if (DeltaFrames >= 1)
+						{
+							UE_LOG(LogRivermax, Warning, TEXT("Output missed %llu frames."), DeltaFrames);
+
+							// If we missed a sync point, this means that last scheduled frame might still be ongoing and 
+							// sending it might be crossing the frame boundary so we skip one entire frame to empty the queue.
+							NextFrameNumber = CurrentFrameNumber + 2;
+						}
+						else
+						{
+							NextFrameNumber = CurrentFrameNumber + 1;
+						}
 					}
 				}
 				else
@@ -1551,24 +1805,27 @@ namespace UE::RivermaxCore::Private
 
 	void FRivermaxOutputStream::CacheCVarValues()
 	{
+		CachedCVars.bEnableCommitTimeProtection = CVarRivermaxOutputEnableTimingProtection.GetValueOnAnyThread() != 0;
 		CachedCVars.ForceCommitImmediateTimeNanosec = CVarRivermaxOutputForceImmediateSchedulingThreshold.GetValueOnAnyThread();
 		CachedCVars.SkipSchedulingTimeNanosec = CVarRivermaxOutputSkipSchedulingCutOffTime.GetValueOnAnyThread() * 1E3;
-		CachedCVars.FrameMemorySliceCount = FMath::Clamp(CVarRivermaxOutputFrameSliceCount.GetValueOnAnyThread(), 1, 100);
-		CachedCVars.bUseIntermediateBuffer = CVarRivermaxOutputEnableIntermediateBuffer.GetValueOnAnyThread() == 1;
 		CachedCVars.bUseSingleMemblock = CVarRivermaxOutputUseSingleMemblock.GetValueOnAnyThread() == 1;
+		CachedCVars.bEnableTimeCriticalThread = CVarRivermaxOutputEnableTimeCriticalThread.GetValueOnAnyThread() != 0;
+		CachedCVars.bShowOutputStats = CVarRivermaxOutputShowStats.GetValueOnAnyThread() != 0;
+		CachedCVars.ShowOutputStatsIntervalSeconds = CVarRivermaxOutputShowStatsInterval.GetValueOnAnyThread();
+		CachedCVars.bPrefillRTPHeaders = CVarRivermaxOutputPrefillRTPHeaders.GetValueOnAnyThread();
 	}
 
-	bool FRivermaxOutputStream::CanProceedWithScheduling() const
+	bool FRivermaxOutputStream::IsChunkOnTime() const
 	{
 		switch (Options.AlignmentMode)
 		{
 		case ERivermaxAlignmentMode::AlignmentPoint:
 		{
-			return CanProceedWithScheduling_AlignmentPoints();
+			return IsChunkOnTime_AlignmentPoints();
 		}
 		case ERivermaxAlignmentMode::FrameCreation:
 		{
-			return CanProceedWithScheduling_FrameCreation();
+			return IsChunkOnTime_FrameCreation();
 		}
 		default:
 		{
@@ -1580,25 +1837,31 @@ namespace UE::RivermaxCore::Private
 		return false;
 	}
 
-	bool FRivermaxOutputStream::CanProceedWithScheduling_FrameCreation() const
+	bool FRivermaxOutputStream::IsChunkOnTime_FrameCreation() const
 	{
 		return true;
 	}
 
-	bool FRivermaxOutputStream::CanProceedWithScheduling_AlignmentPoints() const
+	bool FRivermaxOutputStream::IsChunkOnTime_AlignmentPoints() const
 	{
-		// When aligning on frame boundary, we need to make sure we get here before the scheduling time
-		// There is also some configurable buffer added to make sure we are not right on time
-		// which would cause timing issues.
-		const uint64 CurrentTime = RivermaxModule->GetRivermaxManager()->GetTime();
-		if (StreamData.NextScheduleTimeNanosec <= (CurrentTime + CachedCVars.SkipSchedulingTimeNanosec))
+		if (CachedCVars.bEnableCommitTimeProtection)
 		{
-			UE_LOG(LogRivermax, Warning, TEXT("Can't schedule next frame. Cut off time already passed."));
-			return false;
-		}
+			// Calculate at what time this chunk is supposed to be sent
+			const uint64 DeltaTimePerChunk = (StreamData.FrameFieldTimeIntervalNs - TransmitOffsetNanosec) / StreamMemory.ChunksPerFrameField;
+			const uint64 NextChunkCommitTime = StreamData.NextScheduleTimeNanosec + (CurrentFrame->ChunkNumber * DeltaTimePerChunk);
+
+			// Verify if we are on time to send it. Use CVar to tighten / extend needed window. 
+			// This is to avoid messing up timing
+			const uint64 CurrentTime = RivermaxModule->GetRivermaxManager()->GetTime();
+			if (NextChunkCommitTime <= (CurrentTime + CachedCVars.SkipSchedulingTimeNanosec))
+			{
+				TRACE_CPUPROFILER_EVENT_SCOPE(RmaxOut::ChunkTooLate);
+				return false;
+			}
 
-		// Add other causes of timing issues. 
-		// Possible options : Chunk warnings, Last commit time too close to frame boundary, etc...
+			// Add other causes of timing issues. 
+			// Possible options : Chunk warnings, Last commit time too close to frame boundary, etc...
+		}
 		
 		return true;
 	}
@@ -1614,5 +1877,118 @@ namespace UE::RivermaxCore::Private
 			}
 		}
 	}
+
+	void FRivermaxOutputStream::OnCVarRandomDelayChanged(IConsoleVariable* Var)
+	{
+		bTriggerRandomDelay = true;
+	}
+
+	void FRivermaxOutputStream::CalculateFrameTimestamp()
+	{
+		using namespace UE::RivermaxCore::Private::Utils;
+
+		// For now, in order to be able to use a framelocked input, we pipe frame number in the timestamp for a UE-UE interaction
+		// Follow up work to investigate adding this in RTP header
+		uint64 InputTime = StreamData.NextAlignmentPointNanosec;
+		if (Options.bDoFrameCounterTimestamping)
+		{
+			InputTime = UE::RivermaxCore::GetAlignmentPointFromFrameNumber(CurrentFrame->FrameIdentifier, Options.FrameRate);
+		}
+
+		CurrentFrame->MediaTimestamp = GetTimestampFromTime(InputTime, MediaClockSampleRate);
+	}
+
+	void FRivermaxOutputStream::SkipChunks(uint64 ChunkCount)
+	{
+		bool bHasAddedTrace = false;
+		rmax_status_t Status;
+		do
+		{
+			checkSlow(CachedAPI);
+			Status = CachedAPI->rmax_out_skip_chunks(StreamId, ChunkCount);
+			if (Status != RMAX_OK)
+			{
+				if (Status == RMAX_ERR_NO_FREE_CHUNK)
+				{
+					// Wait until there are enough free chunk to be skipped
+					if (!bHasAddedTrace)
+					{
+						UE_LOG(LogRivermax, Warning, TEXT("No chunks ready to skip. Waiting"));
+						TRACE_CPUPROFILER_EVENT_SCOPE(NoFreeChunk);
+						bHasAddedTrace = true;
+					}
+				}
+				else
+				{
+					ensure(false);
+					UE_LOG(LogRivermax, Error, TEXT("Invalid error happened while trying to skip chunks. Status: %d."), Status);
+					Listener->OnStreamError();
+					Stop();
+				}
+			}
+		} while (Status != RMAX_OK && bIsActive);
+	}
+
+	void FRivermaxOutputStream::SendFrame()
+	{
+		StreamData.LastSendStartTimeNanoSec = RivermaxModule->GetRivermaxManager()->GetTime();
+
+		if (bTriggerRandomDelay)
+		{
+			TRACE_CPUPROFILER_EVENT_SCOPE(FRivermaxOutputStream::RandomDelay);
+			bTriggerRandomDelay = false;
+			FPlatformProcess::SleepNoStats(FMath::RandRange(2e-3, 4e-3));
+		}
+
+		// Calculate frame's timestamp only once and reuse in RTP build
+		CalculateFrameTimestamp();
+
+		const uint32 MediaFrameNumber = Utils::TimestampToFrameNumber(CurrentFrame->MediaTimestamp, Options.FrameRate);
+		TRACE_CPUPROFILER_EVENT_SCOPE_TEXT(*FRivermaxTracingUtils::RmaxOutSendingFrameTraceEvents[MediaFrameNumber % 10]);
+		TRACE_CPUPROFILER_EVENT_SCOPE_TEXT(*FRivermaxTracingUtils::RmaxOutMediaCapturePipeTraceEvents[CurrentFrame->FrameIdentifier % 10]);
+
+		do
+		{
+			if (bIsActive)
+			{
+				GetNextChunk();
+			}
+
+			if (bIsActive && StreamMemory.bUseIntermediateBuffer && ((CurrentFrame->ChunkNumber % StreamMemory.ChunkSpacingBetweenMemcopies) == 0))
+			{
+				CopyFrameData(CurrentFrame, reinterpret_cast<uint8*>(CurrentFrame->FrameStartPtr));
+			}
+
+			if (bIsActive)
+			{
+				SetupRTPHeaders();
+			}
+
+			if (bIsActive)
+			{
+				if (!CurrentFrame->bCaughtTimingIssue)
+				{
+					// As long as our frame is good, verify if we commit chunks before it is expected to be sent
+					// We keep committing the frame even if we detect timing issue to avoid having to skip
+					// chunks in the internals of Rivermax and keep it going for entirety of frames. 
+					// We skip an interval instead but it is quite drastic.
+					const bool bIsChunkOnTime = IsChunkOnTime();
+					CurrentFrame->bCaughtTimingIssue = !bIsChunkOnTime;
+				}
+
+				CommitNextChunks();
+			}
+
+			//Update frame progress
+			if (bIsActive)
+			{
+				Stats.TotalPacketSent += StreamMemory.PacketsPerChunk;
+				++CurrentFrame->ChunkNumber;
+			}
+		} while (CurrentFrame->ChunkNumber < StreamMemory.ChunksPerFrameField && bIsActive);
+
+		Stats.FramesSentCounter++;
+		StreamData.ExpectedFrameIndex = (StreamData.ExpectedFrameIndex + 1) % Options.NumberOfBuffers;
+	}
 }
 
diff --git a/Engine/Plugins/VirtualProduction/Rivermax/RivermaxCore/Source/RivermaxCore/Private/RivermaxOutputStream.h b/Engine/Plugins/VirtualProduction/Rivermax/RivermaxCore/Source/RivermaxCore/Private/RivermaxOutputStream.h
--- a/Engine/Plugins/VirtualProduction/Rivermax/RivermaxCore/Source/RivermaxCore/Private/RivermaxOutputStream.h
+++ b/Engine/Plugins/VirtualProduction/Rivermax/RivermaxCore/Source/RivermaxCore/Private/RivermaxOutputStream.h
@@ -25,43 +25,102 @@ namespace UE::RivermaxCore::Private
 	using UE::RivermaxCore::FRivermaxOutputStreamOptions;
 	using UE::RivermaxCore::FRivermaxOutputVideoFrameInfo;
 
+
+	/** Struct holding configuration information with regards to stream memory and packetization */
 	struct FRivermaxOutputStreamMemory
 	{
+		/** Size of each data payload packet that will be used */
 		uint16 PayloadSize = 0;
+
+		/** Number of pixel group per packet */
 		uint32 PixelGroupPerPacket = 0;
+
+		/** Number of pixels per packet */
 		uint32 PixelsPerPacket = 0;
+
+		/** Number of pixels per frame */
 		uint32 PixelsPerFrame = 0;
 
+		/** Stride of RTP header data. */
 		uint32 HeaderStrideSize = 20;
+
+		/** Number of lines packed inside a chunk. Can be controlled with cvar */
 		uint32 LinesInChunk = 4;
 
+		/** Number of packets per line.  */
 		uint32 PacketsInLine = 0;
+
+		/** Number of packets per chunk. Depends on LinesInChunk */
 		uint32 PacketsPerChunk = 0;
 
+		/** Number of frames per memory block. */
 		uint32 FramesFieldPerMemoryBlock = 0;
+
+		/** Number of packets per frame */
 		uint32 PacketsPerFrame = 0;
+
+		/** Number of packets per memory block */
 		uint32 PacketsPerMemoryBlock = 0;
+
+		/**Number of chunks per frame  */
 		uint32 ChunksPerFrameField = 0;
+
+		/** Number of chunks per memory block */
 		uint32 ChunksPerMemoryBlock = 0;
+
+		/** Number of memory block */
 		uint32 MemoryBlockCount = 0; 
+		
+		/** Whether intermediate buffer is used and captured frame has to be copied over again. */
+		bool bUseIntermediateBuffer = false;
+
+		/** Number of slices we split frame data into when copying it into intermediate buffer */
+		uint32 FrameMemorySliceCount = 1;
+
+		/** Chunk committed between each memcopy of frame data. Helps respect timing. */
+		uint32 ChunkSpacingBetweenMemcopies = 1;
 
+		/** Memory blocks given to Rivermax which where data is located */
 		TArray<rmax_mem_block> MemoryBlocks;
-		TArray<uint16_t> PayloadSizes; //Array describing stride payload size
-		TArray<uint16_t> HeaderSizes; //Array describing header payload size
-		TArray<TArray<FRawRTPHeader>> RTPHeaders; //RTP Headers per memblock
+
+		/** Array with each packet size */
+		TArray<uint16_t> PayloadSizes; 
+
+		/** Array with each RTP header size */
+		TArray<uint16_t> HeaderSizes;
+
+		/** Contains RTP headers per memory block */
+		TArray<TArray<FRawRTPHeader>> RTPHeaders;
+
+		/** Start addresses of each buffer in memblock */
+		TArray<void*> BufferAddresses;
 	};
 
 	struct FRivermaxOutputStreamStats
 	{
+		/** Chunk retries that were required since stream was started */
 		uint32 TotalChunkRetries = 0;
+		
+		/** Chunk retries that happened during last frame */
 		uint32 LastFrameChunkRetries = 0;
+		
+		/** Chunk skipping retries that happened since stream was started */
 		uint32 ChunkSkippingRetries = 0;
-		uint32 TotalStrides = 0;
-		uint32 ChunkWait = 0;
-		uint32 CommitWaits = 0;
+		
+		/** Total packets that have been sent since stream was started */
+		uint32 TotalPacketSent = 0;
+		
+		/** Number of retries that were required when committing and queue was full since stream was started */
 		uint32 CommitRetries = 0;
+
+		/** Immediate commits that were done because we got there too close to scheduling time */
 		uint32 CommitImmediate = 0;
+
+		/** Number of frames that were sent since stream was started */
 		uint64 FramesSentCounter = 0;
+
+		/** Frames that had timing issues since stream was started */
+		uint64 TimingIssueCount = 0;
 	};
 
 	struct FRivermaxOutputStreamData
@@ -70,8 +129,8 @@ namespace UE::RivermaxCore::Private
 		uint32 SequenceNumber = 0;
 		double FrameFieldTimeIntervalNs = 0.0;
 		
-		/** RTP header index we are expected to use. Used to validate returned header pointer by Rivermax. */
-		uint8 RTPHeaderFrameIndex = 0;
+		/** Data and RTP frame index expected to be used for next frame */
+		uint8 ExpectedFrameIndex = 0;
 
 		/** Used to detect misalignment between chunk being sent and frame memory we are writing in */
 		bool bHasFrameFirstChunkBeenFetched = false;
@@ -98,6 +157,9 @@ namespace UE::RivermaxCore::Private
 	/** Struct holding various cached cvar values that can't be changed once stream has been created and to avoid calling anythread getters continuously */
 	struct FOutputStreamCachedCVars
 	{
+		/** Whether timing protection is active and next frame interval is skipped if it happens */
+		bool bEnableCommitTimeProtection = true;
+
 		/** Time padding from scheduling time required to avoid skipping it */
 		uint64 SkipSchedulingTimeNanosec = 0;
 
@@ -107,15 +169,21 @@ namespace UE::RivermaxCore::Private
 		 * gets to actually comitting it. 
 		 */
 		uint64 ForceCommitImmediateTimeNanosec = 0;
-
-		/** Number of slices we split frame data into when copying it into intermediate buffer */
-		uint32 FrameMemorySliceCount = 1;
-		
-		/** Whether intermediate buffer is used and captured frame has to be copied over again. */
-		bool bUseIntermediateBuffer = true;
 		
 		/** Tentative optimization recommended for SDK where a single big memblock is allocated. When false, a memblock per frame is configured. */
 		bool bUseSingleMemblock = true;
+
+		/** Whether to bump output thread priority to time critical */
+		bool bEnableTimeCriticalThread = true;
+
+		/** Whether to show output stats at regular interval in logs */
+		bool bShowOutputStats = false;
+
+		/** Interval in seconds at which to display output stats */
+		float ShowOutputStatsIntervalSeconds = 1.0f;
+
+		/** Whether to prefill RTP header memory with known data at initialization time instead of during sending */
+		bool bPrefillRTPHeaders = true;
 	};
 
 	class FRivermaxOutputStream : public UE::RivermaxCore::IRivermaxOutputStream, public FRunnable
@@ -166,8 +234,8 @@ namespace UE::RivermaxCore::Private
 		/** Destroys rivermax stream. Will wait until it's ready to be destroyed */
 		void DestroyStream();
 
-		/** Waits for the next point in time to send out a new frame */
-		void WaitForNextRound();
+		/** Waits for the next point in time to send out a new frame. Returns true if it exited earlier with the next frame ready to be processed */
+		bool WaitForNextRound();
 
 		/** Calculate next frame scheduling time for alignment points mode */
 		void CalculateNextScheduleTime_AlignementPoints(uint64 CurrentClockTimeNanosec, uint64 CurrentFrameNumber);
@@ -175,20 +243,22 @@ namespace UE::RivermaxCore::Private
 		/** Calculate next frame scheduling time for frame creation mode */
 		void CalculateNextScheduleTime_FrameCreation(uint64 CurrentClockTimeNanosec, uint64 CurrentFrameNumber);
 
-		/** Validates timing post waiting to be sure it is good moving forward with the next frame */
-		bool CanProceedWithScheduling() const;
+		/** Validates timing on every commit to see if we are respecting alignment */
+		bool IsChunkOnTime() const;
 		
-		/** Scheduling verification for frame creation mode. Always valid for now. */
-		bool CanProceedWithScheduling_FrameCreation() const;
+		/** Validates timing for frame creation alignment which always returns true. */
+		bool IsChunkOnTime_FrameCreation() const;
 		
-		/** Scheduling verification for alignment points which validates current timing is ok with alignment points */
-		bool CanProceedWithScheduling_AlignmentPoints() const;
+		/** Validates timing to make sure chunk to be committed are on time. 
+		 *  Once a chunk is late, timings are at risk and next frame will be skipped
+		 */
+		bool IsChunkOnTime_AlignmentPoints() const;
 
 		/** Query rivermax library for the next chunk to work with */
 		void GetNextChunk();
 
 		/** Copies part of frame memory in next memblock's chunk to be sent out */
-		void CopyFrameData();
+		bool CopyFrameData(const TSharedPtr<FRivermaxOutputFrame>& SourceFrame, uint8* DestinationBase);
 
 		/** Fills RTP header for all packets to be sent for this chunk */
 		void SetupRTPHeaders();
@@ -205,6 +275,9 @@ namespace UE::RivermaxCore::Private
 		/** Returns next frame to send for alignement point method. Can return nullptr */
 		void PrepareNextFrame_AlignmentPoint();
 
+		/** Uses time before next frame interval to copy data from next ready frame to intermediate buffer */
+		void PreprocessNextFrame();
+
 		/** If enabled, print stats related to this stream */
 		void ShowStats();
 
@@ -232,6 +305,18 @@ namespace UE::RivermaxCore::Private
 		/** Called back when copy request was completed by allocator */
 		void OnMemoryChunksCopied(const TSharedPtr<FBaseDataCopySideCar>& Sidecar);
 
+		/** Called when delay request cvar has been changed */
+		void OnCVarRandomDelayChanged(IConsoleVariable* Var);
+
+		/** Update frame's timestamp to be used when setting every RTP headers */
+		void CalculateFrameTimestamp();
+
+		/** Tells Rivermax to skip a certain number of chunks in memory. Can be zero to just reset internals */
+		void SkipChunks(uint64 ChunkCount);
+
+		/** Go through all chunks of current frame and commit them to Rivermax to send them at the next desired time */
+		void SendFrame();
+
 	private:
 
 		/** Options related to this stream. i.e resolution, frame rate, etc... */
@@ -303,11 +388,12 @@ namespace UE::RivermaxCore::Private
 		/** Cached cvar values */
 		FOutputStreamCachedCVars CachedCVars;
 
-		/** Time to sleep when waiting for an operation to complete */
-		static constexpr double SleepTimeSeconds = 50.0 * 1E-6;
-
 		/* Pointer to the rivermax API to avoid virtual calls in a hot loop. */ 
 		const UE::RivermaxCore::Private::RIVERMAX_API_FUNCTION_LIST* CachedAPI = nullptr;
+		/** Whether to trigger a delay in the output thread loop next time it ticks */
+		bool bTriggerRandomDelay = false;
+
+		friend struct FRTPHeaderPrefiller;
 	};
 }
 
diff --git a/Engine/Plugins/VirtualProduction/Rivermax/RivermaxCore/Source/RivermaxCore/Private/RivermaxUtils.cpp b/Engine/Plugins/VirtualProduction/Rivermax/RivermaxCore/Source/RivermaxCore/Private/RivermaxUtils.cpp
--- a/Engine/Plugins/VirtualProduction/Rivermax/RivermaxCore/Source/RivermaxCore/Private/RivermaxUtils.cpp
+++ b/Engine/Plugins/VirtualProduction/Rivermax/RivermaxCore/Source/RivermaxCore/Private/RivermaxUtils.cpp
@@ -2,11 +2,17 @@
 
 #include "RivermaxUtils.h"
 
+#include "HAL/IConsoleManager.h"
 #include "RivermaxTypes.h"
 
 
 namespace UE::RivermaxCore::Private::Utils
 {
+	static TAutoConsoleVariable<float> CVarRivermaxOutputFrameRateMultiplier(
+		TEXT("Rivermax.Output.FrameRateMultiplier"), 1.00,
+		TEXT("Multiplier applied to desired output frame rate in order to reduce time it takes to send out a frame and slowly correct misalignment that could happen."),
+		ECVF_Default);
+
 	FString PixelFormatToSamplingDesc(ESamplingType SamplingType)
 	{
 		switch (SamplingType)
@@ -90,14 +96,22 @@ namespace UE::RivermaxCore::Private::Utils
 		// Basic SDP string creation from a set of options. At some point, having a proper SDP loader / creator would be useful.
 		// Refer to https://datatracker.ietf.org/doc/html/rfc4570
 
+		// Apply desired multiplier, if any, to the desired output frame rate.
+		const float FrameRateMultiplier = FMath::Clamp(CVarRivermaxOutputFrameRateMultiplier.GetValueOnAnyThread(), 1.0, 1.2);
 		FString FrameRateDescription;
-		if (FMath::IsNearlyZero(FMath::Frac(Options.FrameRate.AsDecimal())) == false)
+		if (!FMath::IsNearlyEqual(FrameRateMultiplier, 1.0f))
+		{
+			// When a multiplier is involved, always write out frame rate using decimals
+			const double AdjustedFrameRate = Options.FrameRate.AsDecimal() * CVarRivermaxOutputFrameRateMultiplier.GetValueOnAnyThread();
+			FrameRateDescription = FString::Printf(TEXT("%0.9f"), AdjustedFrameRate);
+		}
+		else if (FMath::IsNearlyZero(FMath::Frac(Options.FrameRate.AsDecimal())) == false)
 		{
 			FrameRateDescription = FString::Printf(TEXT("%d/%d"), Options.FrameRate.Numerator, Options.FrameRate.Denominator);
 		}
 		else
 		{
-			FrameRateDescription = FString::Printf(TEXT("%d"), (uint32)Options.FrameRate.AsDecimal());
+			FrameRateDescription = FString::Printf(TEXT("%d"), (int32)Options.FrameRate.AsDecimal());
 		}
 
 
@@ -116,4 +130,12 @@ namespace UE::RivermaxCore::Private::Utils
 			, *FrameRateDescription
 			, *PixelFormatToBitDepth(Options.PixelFormat));
 	}
+
+	uint32 TimestampToFrameNumber(uint32 Timestamp, const FFrameRate& FrameRate)
+	{
+		using namespace UE::RivermaxCore::Private::Utils;
+		const double MediaFrameTime = Timestamp / MediaClockSampleRate;
+		const uint32 FrameNumber = FMath::Floor(MediaFrameTime * FrameRate.AsDecimal());
+		return FrameNumber;
+	}
 }
\ No newline at end of file
diff --git a/Engine/Plugins/VirtualProduction/Rivermax/RivermaxCore/Source/RivermaxCore/Private/RivermaxUtils.h b/Engine/Plugins/VirtualProduction/Rivermax/RivermaxCore/Source/RivermaxCore/Private/RivermaxUtils.h
--- a/Engine/Plugins/VirtualProduction/Rivermax/RivermaxCore/Source/RivermaxCore/Private/RivermaxUtils.h
+++ b/Engine/Plugins/VirtualProduction/Rivermax/RivermaxCore/Source/RivermaxCore/Private/RivermaxUtils.h
@@ -32,4 +32,10 @@ namespace UE::RivermaxCore::Private::Utils
 
 	/** Convert a set of streaming option to its SDP description. Currently only support video type. */
 	void StreamOptionsToSDPDescription(const UE::RivermaxCore::FRivermaxOutputStreamOptions& Options, FAnsiStringBuilderBase& OutSDPDescription);
+
+	/**
+	 * Converts a timestamp in MediaClock period units to a frame number for a given frame rate
+	 * 2110-20 streams uses a standard media clock rate of 90kHz
+	 */
+	uint32 TimestampToFrameNumber(uint32 Timestamp, const FFrameRate& FrameRate);
 }
diff --git a/Engine/Plugins/VirtualProduction/Rivermax/RivermaxCore/Source/RivermaxCore/Public/RivermaxTracingUtils.h b/Engine/Plugins/VirtualProduction/Rivermax/RivermaxCore/Source/RivermaxCore/Public/RivermaxTracingUtils.h
new file mode 100644
--- /dev/null
+++ b/Engine/Plugins/VirtualProduction/Rivermax/RivermaxCore/Source/RivermaxCore/Public/RivermaxTracingUtils.h
@@ -0,0 +1,34 @@
+// Copyright Epic Games, Inc. All Rights Reserved.
+
+#pragma once
+
+#include "Containers/UnrealString.h"
+#include "Containers/Map.h"
+#include "CoreTypes.h"
+
+namespace UE::RivermaxCore
+{
+	class RIVERMAXCORE_API FRivermaxTracingUtils
+	{
+	public:
+		/** Media capture tracing string used to track a specific frame number being captured */
+		static TMap<uint8, FString> RmaxOutMediaCapturePipeTraceEvents;
+
+		/** Rivermax output stream sending markup */
+		static TMap<uint8, FString> RmaxOutSendingFrameTraceEvents;
+
+		/** Rivermax frame capture is ready markup */
+		static TMap<uint8, FString> RmaxOutFrameReadyTraceEvents;
+
+		/** Rivermax input stream starting to receive a frame markup */
+		static TMap<uint8, FString> RmaxInStartingFrameTraceEvents;
+
+		/** Rivermax input stream frame received markup */
+		static TMap<uint8, FString> RmaxInReceivedFrameTraceEvents;
+
+		/** Rivermax player frame selected markup */
+		static TMap<uint8, FString> RmaxInSelectedFrameTraceEvents;
+	};
+}
+
+
diff --git a/Engine/Plugins/VirtualProduction/Rivermax/RivermaxMedia/Source/RivermaxMedia/Private/RivermaxMediaCapture.cpp b/Engine/Plugins/VirtualProduction/Rivermax/RivermaxMedia/Source/RivermaxMedia/Private/RivermaxMediaCapture.cpp
--- a/Engine/Plugins/VirtualProduction/Rivermax/RivermaxMedia/Source/RivermaxMedia/Private/RivermaxMediaCapture.cpp
+++ b/Engine/Plugins/VirtualProduction/Rivermax/RivermaxMedia/Source/RivermaxMedia/Private/RivermaxMediaCapture.cpp
@@ -12,6 +12,7 @@
 #include "RivermaxMediaOutput.h"
 #include "RivermaxMediaUtils.h"
 #include "RivermaxShaders.h"
+#include "RivermaxTracingUtils.h"
 #include "Slate/SceneViewport.h"
 #include "Widgets/SViewport.h"
 
@@ -270,7 +271,7 @@ void URivermaxMediaCapture::AddFrameReservationPass(FRDGBuilder& GraphBuilder)
 	GraphBuilder.RHICmdList.EnqueueLambda([Capturer, FrameCounter = GFrameCounterRenderThread](FRHICommandList& RHICmdList)
 	{
 		TRACE_CPUPROFILER_EVENT_SCOPE(RmaxFrameReservation);
-		TRACE_CPUPROFILER_EVENT_SCOPE_TEXT(*FString::Printf(TEXT("MediaCapturePipe: %llu"), FrameCounter));
+		TRACE_CPUPROFILER_EVENT_SCOPE_TEXT(*UE::RivermaxCore::FRivermaxTracingUtils::RmaxOutMediaCapturePipeTraceEvents[FrameCounter % 10]);
 
 		if (Capturer->bIsActive)
 		{
@@ -302,7 +303,7 @@ void URivermaxMediaCapture::OnRHIResourceCaptured_AnyThread(const FCaptureBaseDa
 	using namespace UE::RivermaxCore;
 
 	TRACE_CPUPROFILER_EVENT_SCOPE(URivermaxMediaCapture::OnRHIResourceCaptured_AnyThread);
-	TRACE_CPUPROFILER_EVENT_SCOPE_TEXT(*FString::Printf(TEXT("MediaCapturePipe: %u"), InBaseData.SourceFrameNumberRenderThread));
+	TRACE_CPUPROFILER_EVENT_SCOPE_TEXT(*FRivermaxTracingUtils::RmaxOutMediaCapturePipeTraceEvents[InBaseData.SourceFrameNumberRenderThread % 10]);
 
 	FRivermaxOutputVideoFrameInfo NewFrame;
 	NewFrame.FrameIdentifier = InBaseData.SourceFrameNumberRenderThread;
@@ -318,7 +319,7 @@ void URivermaxMediaCapture::OnFrameCaptured_AnyThread(const FCaptureBaseData& In
 	using namespace UE::RivermaxCore;
 
 	TRACE_CPUPROFILER_EVENT_SCOPE(URivermaxMediaCapture::OnFrameCaptured_AnyThread);
-	TRACE_CPUPROFILER_EVENT_SCOPE_TEXT(*FString::Printf(TEXT("MediaCapturePipe: %u"), InBaseData.SourceFrameNumberRenderThread));
+	TRACE_CPUPROFILER_EVENT_SCOPE_TEXT(*FRivermaxTracingUtils::RmaxOutMediaCapturePipeTraceEvents[InBaseData.SourceFrameNumberRenderThread % 10]);
 	
 	FRivermaxOutputVideoFrameInfo NewFrame;
 	NewFrame.Height = InResourceData.Height;
@@ -380,8 +381,7 @@ void URivermaxMediaCapture::OnCustomCapture_RenderingThread(FRDGBuilder& GraphBu
 {
 	RDG_GPU_STAT_SCOPE(GraphBuilder, Rmax_Capture)
 	TRACE_CPUPROFILER_EVENT_SCOPE(URivermaxMediaCapture::OnCustomCapture_RenderingThread);
-	TRACE_CPUPROFILER_EVENT_SCOPE_TEXT(*FString::Printf(TEXT("MediaCapturePipe: %llu"), GFrameCounterRenderThread));
-
+	TRACE_CPUPROFILER_EVENT_SCOPE_TEXT(*UE::RivermaxCore::FRivermaxTracingUtils::RmaxOutMediaCapturePipeTraceEvents[GFrameCounterRenderThread % 10]);
 
 	using namespace UE::RivermaxShaders;
 	URivermaxMediaOutput* RivermaxOutput = CastChecked<URivermaxMediaOutput>(MediaOutput);
diff --git a/Engine/Plugins/VirtualProduction/Rivermax/RivermaxMedia/Source/RivermaxMedia/Private/RivermaxMediaPlayer.cpp b/Engine/Plugins/VirtualProduction/Rivermax/RivermaxMedia/Source/RivermaxMedia/Private/RivermaxMediaPlayer.cpp
--- a/Engine/Plugins/VirtualProduction/Rivermax/RivermaxMedia/Source/RivermaxMedia/Private/RivermaxMediaPlayer.cpp
+++ b/Engine/Plugins/VirtualProduction/Rivermax/RivermaxMedia/Source/RivermaxMedia/Private/RivermaxMediaPlayer.cpp
@@ -17,6 +17,7 @@
 #include "RivermaxMediaTextureSample.h"
 #include "RivermaxMediaUtils.h"
 #include "RivermaxPTPUtils.h"
+#include "RivermaxTracingUtils.h"
 #include "RivermaxTypes.h"
 #include "Stats/Stats2.h"
 #include "Tasks/Task.h"
@@ -711,8 +712,7 @@ namespace UE::RivermaxMedia
 			OutExpectation.FrameIndex = HighestTimestampIndex;
 
 			// Make the ones skipped available for reception again
-			const FString TraceName = FString::Format(TEXT("RmaxInput::Selected {0}"), { OutExpectation.FrameNumber });
-			TRACE_CPUPROFILER_EVENT_SCOPE_TEXT(*TraceName);
+			TRACE_CPUPROFILER_EVENT_SCOPE_TEXT(*FRivermaxTracingUtils::RmaxInSelectedFrameTraceEvents[OutExpectation.FrameNumber % 10]);
 
 			for (const TSharedPtr<FRivermaxSampleWrapper>& Frame : SamplePool)
 			{
@@ -990,6 +990,7 @@ namespace UE::RivermaxMedia
 			FPlatformProcess::SleepNoStats(SleepTimeSeconds);
 			if ((FPlatformTime::Seconds() - StartTimeSeconds) > TimeoutSeconds)
 			{
+				UE_LOG(LogRivermaxMedia, Warning, TEXT("Timed out waiting for pendings tasks to finish."));
 				break;
 			}
 		}
