- Fix for 2110 output streams timing getting broken when launching heavier projects with 2 cameras. CVars added to configure setting.

diff --git a/Engine/Plugins/VirtualProduction/Rivermax/RivermaxCore/Source/RivermaxCore/Private/RivermaxOutputStream.cpp b/Engine/Plugins/VirtualProduction/Rivermax/RivermaxCore/Source/RivermaxCore/Private/RivermaxOutputStream.cpp
--- a/Engine/Plugins/VirtualProduction/Rivermax/RivermaxCore/Source/RivermaxCore/Private/RivermaxOutputStream.cpp
+++ b/Engine/Plugins/VirtualProduction/Rivermax/RivermaxCore/Source/RivermaxCore/Private/RivermaxOutputStream.cpp
@@ -112,6 +112,18 @@ namespace UE::RivermaxCore::Private
 		TEXT("Optimization used to prefill every RTP headers with known data."),
 		ECVF_Default);
 
+
+	static TAutoConsoleVariable<bool> CVarRivermaxOutputForceSkip(
+		TEXT("Rivermax.Output.ForceSkip"), true,
+		TEXT("Used to prevent enforced blank window when multiplier is used to prevent timing issues.")
+		TEXT("Only affects alignment point method."),
+		ECVF_Default);
+
+	static TAutoConsoleVariable<float> CVarRivermaxOutputFrameRateMultiplier(
+		TEXT("Rivermax.Output.FrameRateMultiplier"), 1.05,
+		TEXT("Multiplier applied to desired output frame rate in order to reduce time it takes to send out a frame and slowly correct misalignment that could happen."),
+		ECVF_Default);
+
 	bool FindPayloadSize(const FRivermaxOutputStreamOptions& InOptions, uint32 InBytesPerLine, const FVideoFormatInfo& FormatInfo, uint16& OutPayloadSize)
 	{
 		using namespace UE::RivermaxCore::Private::Utils;
@@ -501,13 +513,15 @@ namespace UE::RivermaxCore::Private
 			return false;
 		}
 
+		InitializeTimingProtections();
+
 		// Enable frame boundary monitoring
 		MonitoringGuid = RivermaxModule->GetRivermaxBoundaryMonitor().StartMonitoring(Options.FrameRate);
 
 		InitializationFuture = Async(EAsyncExecution::TaskGraph, [this]()
 		{
 			TAnsiStringBuilder<2048> SDPDescription;
-			UE::RivermaxCore::Private::Utils::StreamOptionsToSDPDescription(Options, SDPDescription);
+			UE::RivermaxCore::Private::Utils::StreamOptionsToSDPDescription(Options, StreamMemory.FrameRateMultiplier, SDPDescription);
 
 			// Create Rivermax stream using memory configuration
 			{
@@ -672,7 +686,10 @@ namespace UE::RivermaxCore::Private
 				if (CurrentFrame->bCaughtTimingIssue)
 				{
 					++Stats.TimingIssueCount;
-					
+				}
+
+				if (CurrentFrame->bCaughtTimingIssue || StreamMemory.bAlwaysSkipChunk)
+				{
 					constexpr uint64 ChunksToSkip = 0;
 					SkipChunks(ChunksToSkip);
 				}
@@ -1990,5 +2007,29 @@ namespace UE::RivermaxCore::Private
 		Stats.FramesSentCounter++;
 		StreamData.ExpectedFrameIndex = (StreamData.ExpectedFrameIndex + 1) % Options.NumberOfBuffers;
 	}
+
+	void FRivermaxOutputStream::InitializeTimingProtections()
+	{
+		switch (Options.AlignmentMode)
+		{
+			case ERivermaxAlignmentMode::FrameCreation:
+			{
+				StreamMemory.bAlwaysSkipChunk = false;
+				StreamMemory.FrameRateMultiplier = 1.0f;
+				break;
+			}
+			case ERivermaxAlignmentMode::AlignmentPoint:
+			{
+				StreamMemory.bAlwaysSkipChunk = CVarRivermaxOutputForceSkip.GetValueOnAnyThread();
+				StreamMemory.FrameRateMultiplier = FMath::Clamp(CVarRivermaxOutputFrameRateMultiplier.GetValueOnAnyThread(), 1.0, 2.0);
+				break;
+			}
+			default:
+			{
+				checkNoEntry();
+			}
+		}
+	}
+
 }
 
diff --git a/Engine/Plugins/VirtualProduction/Rivermax/RivermaxCore/Source/RivermaxCore/Private/RivermaxOutputStream.h b/Engine/Plugins/VirtualProduction/Rivermax/RivermaxCore/Source/RivermaxCore/Private/RivermaxOutputStream.h
--- a/Engine/Plugins/VirtualProduction/Rivermax/RivermaxCore/Source/RivermaxCore/Private/RivermaxOutputStream.h
+++ b/Engine/Plugins/VirtualProduction/Rivermax/RivermaxCore/Source/RivermaxCore/Private/RivermaxOutputStream.h
@@ -80,6 +80,19 @@ namespace UE::RivermaxCore::Private
 		/** Chunk committed between each memcopy of frame data. Helps respect timing. */
 		uint32 ChunkSpacingBetweenMemcopies = 1;
 
+		/** 
+		 * Whether we skip zero chunk, to reset rivermax's stream internal timings, vblank verification, etc...
+		 * Alignment point method will default to true and be configurable through cvar and frame creation will never use it.
+		 * Usage is to make sure Tro is respected when frame rate multiplier is involved.
+		 */
+		bool bAlwaysSkipChunk = true;
+
+		/** 
+		 * Multiplier applied to desired frame rate when creating Rivermax stream to add margin in case of packet timings not being respected
+		 * This will break compliance with standard but if receiver can tolerate it, more stable stream timings are expected.
+		 */
+		float FrameRateMultiplier = 1.00;
+
 		/** Memory blocks given to Rivermax which where data is located */
 		TArray<rmax_mem_block> MemoryBlocks;
 
@@ -219,6 +232,9 @@ namespace UE::RivermaxCore::Private
 		/** Initializes timing setup for this stream. TRO, frame interval etc... */
 		void InitializeStreamTimingSettings();
 
+		/** Configures settings related to timing protection dependant on cvars */
+		void InitializeTimingProtections();
+
 		/** Sets up frame management taking care of allocation, special cuda handling, etc... */
 		bool SetupFrameManagement();
 
diff --git a/Engine/Plugins/VirtualProduction/Rivermax/RivermaxCore/Source/RivermaxCore/Private/RivermaxUtils.cpp b/Engine/Plugins/VirtualProduction/Rivermax/RivermaxCore/Source/RivermaxCore/Private/RivermaxUtils.cpp
--- a/Engine/Plugins/VirtualProduction/Rivermax/RivermaxCore/Source/RivermaxCore/Private/RivermaxUtils.cpp
+++ b/Engine/Plugins/VirtualProduction/Rivermax/RivermaxCore/Source/RivermaxCore/Private/RivermaxUtils.cpp
@@ -3,16 +3,12 @@
 #include "RivermaxUtils.h"
 
 #include "HAL/IConsoleManager.h"
+#include "RivermaxLog.h"
 #include "RivermaxTypes.h"
 
 
 namespace UE::RivermaxCore::Private::Utils
 {
-	static TAutoConsoleVariable<float> CVarRivermaxOutputFrameRateMultiplier(
-		TEXT("Rivermax.Output.FrameRateMultiplier"), 1.00,
-		TEXT("Multiplier applied to desired output frame rate in order to reduce time it takes to send out a frame and slowly correct misalignment that could happen."),
-		ECVF_Default);
-
 	FString PixelFormatToSamplingDesc(ESamplingType SamplingType)
 	{
 		switch (SamplingType)
@@ -91,30 +87,37 @@ namespace UE::RivermaxCore::Private::Utils
 		}
 	}
 
-	void StreamOptionsToSDPDescription(const UE::RivermaxCore::FRivermaxOutputStreamOptions& Options, FAnsiStringBuilderBase& OutSDPDescription)
+	void StreamOptionsToSDPDescription(const UE::RivermaxCore::FRivermaxOutputStreamOptions& Options, float RateMultiplier, FAnsiStringBuilderBase& OutSDPDescription)
 	{
 		// Basic SDP string creation from a set of options. At some point, having a proper SDP loader / creator would be useful.
 		// Refer to https://datatracker.ietf.org/doc/html/rfc4570
 
-		// Apply desired multiplier, if any, to the desired output frame rate.
-		const float FrameRateMultiplier = FMath::Clamp(CVarRivermaxOutputFrameRateMultiplier.GetValueOnAnyThread(), 1.0, 1.2);
+		// Apply desired multiplier and convert fractional frame rate to be represented over 1001 for sdp compliance.
+		FFrameRate DesiredRate = Options.FrameRate;
+		const double DecimalDesiredRate = Options.FrameRate.AsDecimal();
+		const bool bIsNonStandardFractionalFrameRate = (DesiredRate.Denominator != 1001 && !FMath::IsNearlyZero(FMath::Frac(DecimalDesiredRate)));
+		if (bIsNonStandardFractionalFrameRate)
+		{
+			UE_LOG(LogRivermax, Warning, TEXT("Fractional frame rates must be described using a denominator of 1001. Converting it for stream creation."));
+		}
+
 		FString FrameRateDescription;
-		if (!FMath::IsNearlyEqual(FrameRateMultiplier, 1.0f))
+		if (!FMath::IsNearlyEqual(RateMultiplier, 1.0f) || bIsNonStandardFractionalFrameRate)
 		{
-			// When a multiplier is involved, always write out frame rate using decimals
-			const double AdjustedFrameRate = Options.FrameRate.AsDecimal() * CVarRivermaxOutputFrameRateMultiplier.GetValueOnAnyThread();
-			FrameRateDescription = FString::Printf(TEXT("%0.9f"), AdjustedFrameRate);
+			const double NewRateDecimal = DecimalDesiredRate * RateMultiplier * 1001;
+			DesiredRate.Numerator = FMath::RoundToInt32(NewRateDecimal);
+			DesiredRate.Denominator = 1001;
 		}
-		else if (FMath::IsNearlyZero(FMath::Frac(Options.FrameRate.AsDecimal())) == false)
+		
+		if(DesiredRate.Denominator == 1001)
 		{
-			FrameRateDescription = FString::Printf(TEXT("%d/%d"), Options.FrameRate.Numerator, Options.FrameRate.Denominator);
+			FrameRateDescription = FString::Printf(TEXT("%u/%u"), DesiredRate.Numerator, DesiredRate.Denominator);
 		}
 		else
 		{
-			FrameRateDescription = FString::Printf(TEXT("%d"), (int32)Options.FrameRate.AsDecimal());
+			FrameRateDescription = FString::Printf(TEXT("%u"), FMath::RoundToInt32(DesiredRate.AsDecimal()));
 		}
 
-
 		constexpr int32 MulticastTTL = 64;
 		OutSDPDescription.Appendf("v=0\n");
 		OutSDPDescription.Appendf("s=SMPTE ST2110 20 streams\n");
diff --git a/Engine/Plugins/VirtualProduction/Rivermax/RivermaxCore/Source/RivermaxCore/Private/RivermaxUtils.h b/Engine/Plugins/VirtualProduction/Rivermax/RivermaxCore/Source/RivermaxCore/Private/RivermaxUtils.h
--- a/Engine/Plugins/VirtualProduction/Rivermax/RivermaxCore/Source/RivermaxCore/Private/RivermaxUtils.h
+++ b/Engine/Plugins/VirtualProduction/Rivermax/RivermaxCore/Source/RivermaxCore/Private/RivermaxUtils.h
@@ -31,7 +31,7 @@ namespace UE::RivermaxCore::Private::Utils
 	static constexpr float SleepTimeSeconds = 50 * 1E-6;
 
 	/** Convert a set of streaming option to its SDP description. Currently only support video type. */
-	void StreamOptionsToSDPDescription(const UE::RivermaxCore::FRivermaxOutputStreamOptions& Options, FAnsiStringBuilderBase& OutSDPDescription);
+	void StreamOptionsToSDPDescription(const UE::RivermaxCore::FRivermaxOutputStreamOptions& Options, float RateMultiplier, FAnsiStringBuilderBase& OutSDPDescription);
 
 	/**
 	 * Converts a timestamp in MediaClock period units to a frame number for a given frame rate
