Adding critical section during rivermax input stream initialization to avoid having multiple stream init happening at the same time in the library

diff --git a/Engine/Plugins/VirtualProduction/Rivermax/RivermaxCore/Source/RivermaxCore/Private/RivermaxInputStream.cpp b/Engine/Plugins/VirtualProduction/Rivermax/RivermaxCore/Source/RivermaxCore/Private/RivermaxInputStream.cpp
--- a/Engine/Plugins/VirtualProduction/Rivermax/RivermaxCore/Source/RivermaxCore/Private/RivermaxInputStream.cpp
+++ b/Engine/Plugins/VirtualProduction/Rivermax/RivermaxCore/Source/RivermaxCore/Private/RivermaxInputStream.cpp
@@ -4,6 +4,7 @@
 
 #include "Async/Async.h"
 #include "CudaModule.h"
+#include "HAL/CriticalSection.h"
 #include "ID3D12DynamicRHI.h"
 #include "IRivermaxCoreModule.h"
 #include "IRivermaxManager.h"
@@ -19,6 +20,9 @@
 
 namespace UE::RivermaxCore::Private
 {
+	/** Initialization calls to Rivermax SDK isn't threadsafe and to prevent stalling game thread, a critical section is shared across streams */
+	static FCriticalSection InitCriticalSection;
+
 	static TAutoConsoleVariable<float> CVarWaitForCompletionTimeout(
 		TEXT("Rivermax.Input.WaitCompletionTimeout"),
 		0.25,
@@ -134,6 +138,9 @@ namespace UE::RivermaxCore::Private
 
 		InitTaskFuture = Async(EAsyncExecution::TaskGraph, [this]()
 		{
+			TRACE_CPUPROFILER_EVENT_SCOPE(FRivermaxInputStream::InitStream);
+			FScopeLock Lock(&InitCriticalSection);
+
 			// If the stream is trying to shutdown before the init task has even started, don't bother
 			if (bIsShuttingDown)
 			{
