- Adding first packet interval stats

diff --git a/Engine/Plugins/VirtualProduction/Rivermax/RivermaxCore/Source/RivermaxCore/Private/RivermaxInputStream.cpp b/Engine/Plugins/VirtualProduction/Rivermax/RivermaxCore/Source/RivermaxCore/Private/RivermaxInputStream.cpp
--- a/Engine/Plugins/VirtualProduction/Rivermax/RivermaxCore/Source/RivermaxCore/Private/RivermaxInputStream.cpp
+++ b/Engine/Plugins/VirtualProduction/Rivermax/RivermaxCore/Source/RivermaxCore/Private/RivermaxInputStream.cpp
@@ -10,6 +10,7 @@
 #include "IRivermaxManager.h"
 #include "Misc/ByteSwap.h"
 #include "RivermaxLog.h"
+#include "RivermaxPTPUtils.h"
 #include "RivermaxTracingUtils.h"
 #include "RivermaxUtils.h"
 #include "RTPHeader.h"
@@ -35,6 +36,28 @@ namespace UE::RivermaxCore::Private
 		TEXT("Expected payload size used to initialize rivermax stream."),
 		ECVF_Default);
 
+	static bool GEnableLargeFirstPacketIntervalLogging = false;
+	FAutoConsoleVariableRef CVarRivermaxEnableLargeFirstPacketIntervalLogging(
+		TEXT("Rivermax.Input.EnableLargeFistPacketIntervalLogging")
+		, UE::RivermaxCore::Private::GEnableLargeFirstPacketIntervalLogging
+		, TEXT("Enables detection and logging of large delta times between frame boundary and first packet reception.")
+		, ECVF_Default);
+
+	static bool GClearFirstPacketIntervalStats = false;
+	FAutoConsoleVariableRef CVarRivermaxClearFirstPacketIntervalStats(
+		TEXT("Rivermax.Input.ClearFistPacketIntervalStats")
+		, UE::RivermaxCore::Private::GClearFirstPacketIntervalStats
+		, TEXT("Clears stats related to first packet interval detection.")
+		, ECVF_Default);
+
+	static int32 GLargeFirstPacketIntervalThresholdMicroSec = 2000;
+	FAutoConsoleVariableRef CVarRivermaxLargeFirstPacketIntervalThreshold(
+		TEXT("Rivermax.Input.LargeFistPacketIntervalThreshold")
+		, UE::RivermaxCore::Private::GLargeFirstPacketIntervalThresholdMicroSec
+		, TEXT("Microseconds required to consider a first packet interval to be large and be logged.")
+		, ECVF_Default);
+
+
 	void FFrameDescriptionTrackingData::ResetSingleFrameTracking()
 	{
 		PayloadSizeReceived.Empty();
@@ -337,11 +360,48 @@ namespace UE::RivermaxCore::Private
 					{
 						const uint32 FrameNumber = Utils::TimestampToFrameNumber(RTPHeader.Timestamp, Options.FrameRate);
 						TRACE_CPUPROFILER_EVENT_SCOPE_TEXT(*FRivermaxTracingUtils::RmaxInStartingFrameTraceEvents[FrameNumber % 10]);
+
+						IRivermaxCoreModule& RivermaxModule = FModuleManager::LoadModuleChecked<IRivermaxCoreModule>(TEXT("RivermaxCore"));
+						const uint64 CurrentTime = RivermaxModule.GetRivermaxManager()->GetTime();
+
 						bIsFirstPacketReceived = true;
+
+						if (StreamStats.EndOfFrameReceived > 0)
+						{
+							const uint64 FrameBoundary = GetAlignmentPointFromFrameNumber(GetFrameNumber(CurrentTime, Options.FrameRate), Options.FrameRate);
+							const uint64 FirstPacketInterval = CurrentTime - FrameBoundary;
+							StreamStats.MinFirstPacketIntervalNS = FMath::Min(StreamStats.MinFirstPacketIntervalNS, FirstPacketInterval);
+							StreamStats.MaxFirstPacketIntervalNS = FMath::Max(StreamStats.MaxFirstPacketIntervalNS, FirstPacketInterval);
+							
+							constexpr double Alpha = 0.8;
+							StreamStats.FirstPacketIntervalAccumulatorNS = Alpha * FirstPacketInterval + ((1.0 - Alpha) * StreamStats.FirstPacketIntervalAccumulatorNS);
+
+							if (GClearFirstPacketIntervalStats)
+							{
+								GClearFirstPacketIntervalStats = false;
+								StreamStats.MinFirstPacketIntervalNS = TNumericLimits<uint64>::Max();
+								StreamStats.MaxFirstPacketIntervalNS = TNumericLimits<uint64>::Min();
+								StreamStats.FirstPacketIntervalAccumulatorNS = 0;
+							}
+
+							if (GEnableLargeFirstPacketIntervalLogging)
+							{
+								const uint32 IntervalMicroSec = FirstPacketInterval / 1000;
+								if (IntervalMicroSec > (uint32)GLargeFirstPacketIntervalThresholdMicroSec)
+								{
+
+									UE_LOG(LogRivermax, Warning, TEXT("Large First packet interval: %llu. Min: %llu. Max: %llu. Avg: %llu.")
+									, FirstPacketInterval
+									, StreamStats.MinFirstPacketIntervalNS
+									, StreamStats.MaxFirstPacketIntervalNS
+									, StreamStats.FirstPacketIntervalAccumulatorNS);
+								}
+							}
+						}
 					}
 
 					StreamStats.BytesReceived += PacketSize + HeaderSize;
-
+					
 					UpdateFrameTracking(RTPHeader);
 
 					switch (State)
@@ -443,8 +503,10 @@ namespace UE::RivermaxCore::Private
 		if (CurrentTime - LastLoggingTimestamp >= LoggingInterval)
 		{
 			LastLoggingTimestamp = CurrentTime;
-			UE_LOG(LogRivermax, Verbose, TEXT("Stream %d stats: FrameCount: %llu, EndOfFrame: %llu, Chunks: %llu, Bytes: %llu, PacketLossInFrame: %llu, TotalPacketLoss: %llu, BiggerFrames: %llu, InvalidFrames: %llu, InvalidHeader: %llu, EmptyCompletion: %llu")
+			UE_LOG(LogRivermax, Verbose, TEXT("Stream %d (%s:%u) stats: FrameCount: %llu, EndOfFrame: %llu, Chunks: %llu, Bytes: %llu, PacketLossInFrame: %llu, TotalPacketLoss: %llu, BiggerFrames: %llu, InvalidFrames: %llu, InvalidHeader: %llu, EmptyCompletion: %llu")
 			, StreamId
+			, *Options.StreamAddress
+			, Options.Port
 			, StreamStats.FramesReceived
 			, StreamStats.EndOfFrameReceived
 			, StreamStats.ChunksReceived
@@ -456,6 +518,15 @@ namespace UE::RivermaxCore::Private
 			, StreamStats.InvalidHeadercount
 			, StreamStats.EmptyCompletionCount
 			);
+
+			UE_LOG(LogRivermax, Verbose, TEXT("Stream %d (%s:%u) first packet interval: - Min: %llu. Max: %llu. Avg: %llu.")
+				, StreamId
+				, *Options.StreamAddress
+				, Options.Port
+				, StreamStats.MinFirstPacketIntervalNS
+				, StreamStats.MaxFirstPacketIntervalNS
+				, StreamStats.FirstPacketIntervalAccumulatorNS);
+
 		}
 	}
 
@@ -945,7 +1016,7 @@ namespace UE::RivermaxCore::Private
 			// No need to provide the new frame and prepare the next one if we are shutting down
 			if (bIsShuttingDown == false)
 			{
-				UE_LOG(LogRivermax, Verbose, TEXT("RmaxRX frame number %u with timestamp %u."), Descriptor.FrameNumber, Descriptor.Timestamp);
+				UE_LOG(LogRivermax, VeryVerbose, TEXT("RmaxRX frame number %u with timestamp %u."), Descriptor.FrameNumber, Descriptor.Timestamp);
 				
 				FRivermaxInputVideoFrameReception NewFrame;
 				NewFrame.VideoBuffer = reinterpret_cast<uint8*>(StreamData.CurrentFrame);
diff --git a/Engine/Plugins/VirtualProduction/Rivermax/RivermaxCore/Source/RivermaxCore/Private/RivermaxInputStream.h b/Engine/Plugins/VirtualProduction/Rivermax/RivermaxCore/Source/RivermaxCore/Private/RivermaxInputStream.h
--- a/Engine/Plugins/VirtualProduction/Rivermax/RivermaxCore/Source/RivermaxCore/Private/RivermaxInputStream.h
+++ b/Engine/Plugins/VirtualProduction/Rivermax/RivermaxCore/Source/RivermaxCore/Private/RivermaxInputStream.h
@@ -63,6 +63,10 @@ namespace UE::RivermaxCore::Private
 		uint64 ChunksReceived = 0;
 		uint64 EndOfFrameReceived = 0;
 		uint64 EmptyCompletionCount = 0;
+		uint64 MinFirstPacketIntervalNS = TNumericLimits<uint64>::Max();
+		uint64 MaxFirstPacketIntervalNS = TNumericLimits<uint64>::Min();
+		uint64 FirstPacketIntervalAccumulatorNS = 0;
+		uint64 FirstPacketIntervalStatsCount = 0;
 	};
 
 	struct FInputStreamData
