Rmax: Changed default rate multiplier from 1.05 to 1.00, since the multiplier was upsetting the receiver, and it seems to not be needed anymore.

diff --git a/Engine/Plugins/VirtualProduction/Rivermax/RivermaxCore/Source/RivermaxCore/Private/RivermaxOutputStream.cpp b/Engine/Plugins/VirtualProduction/Rivermax/RivermaxCore/Source/RivermaxCore/Private/RivermaxOutputStream.cpp
--- a/Engine/Plugins/VirtualProduction/Rivermax/RivermaxCore/Source/RivermaxCore/Private/RivermaxOutputStream.cpp
+++ b/Engine/Plugins/VirtualProduction/Rivermax/RivermaxCore/Source/RivermaxCore/Private/RivermaxOutputStream.cpp
@@ -120,7 +120,7 @@ namespace UE::RivermaxCore::Private
 		ECVF_Default);
 
 	static TAutoConsoleVariable<float> CVarRivermaxOutputFrameRateMultiplier(
-		TEXT("Rivermax.Output.FrameRateMultiplier"), 1.05,
+		TEXT("Rivermax.Output.FrameRateMultiplier"), 1.0,
 		TEXT("Multiplier applied to desired output frame rate in order to reduce time it takes to send out a frame and slowly correct misalignment that could happen."),
 		ECVF_Default);
 
