Rivermax: Added frame number log at the point of sending so it can be compared with the frame numbers on the receiving side

diff --git a/Engine/Plugins/VirtualProduction/Rivermax/RivermaxCore/Source/RivermaxCore/Private/RivermaxOutputStream.cpp b/Engine/Plugins/VirtualProduction/Rivermax/RivermaxCore/Source/RivermaxCore/Private/RivermaxOutputStream.cpp
--- a/Engine/Plugins/VirtualProduction/Rivermax/RivermaxCore/Source/RivermaxCore/Private/RivermaxOutputStream.cpp
+++ b/Engine/Plugins/VirtualProduction/Rivermax/RivermaxCore/Source/RivermaxCore/Private/RivermaxOutputStream.cpp
@@ -2042,6 +2042,7 @@ namespace UE::RivermaxCore::Private
 		const uint32 MediaFrameNumber = Utils::TimestampToFrameNumber(CurrentFrame->MediaTimestamp, Options.FrameRate);
 		TRACE_CPUPROFILER_EVENT_SCOPE_TEXT(*FRivermaxTracingUtils::RmaxOutSendingFrameTraceEvents[MediaFrameNumber % 10]);
 		TRACE_CPUPROFILER_EVENT_SCOPE_TEXT(*FRivermaxTracingUtils::RmaxOutMediaCapturePipeTraceEvents[CurrentFrame->FrameIdentifier % 10]);
+		UE_LOG(LogRivermax, VeryVerbose, TEXT("RmaxRX Sending frame number %u with timestamp %u."), MediaFrameNumber, CurrentFrame->MediaTimestamp);
 
 		do
 		{
