Rivermax: When Flow is unable to be attached the rivermax input stream thread was still being created causing Unreal Engine to freeze and cauze severe memory leaks.

diff --git a/Engine/Plugins/VirtualProduction/Rivermax/RivermaxCore/Source/RivermaxCore/Private/RivermaxInputStream.cpp b/Engine/Plugins/VirtualProduction/Rivermax/RivermaxCore/Source/RivermaxCore/Private/RivermaxInputStream.cpp
--- a/Engine/Plugins/VirtualProduction/Rivermax/RivermaxCore/Source/RivermaxCore/Private/RivermaxInputStream.cpp
+++ b/Engine/Plugins/VirtualProduction/Rivermax/RivermaxCore/Source/RivermaxCore/Private/RivermaxInputStream.cpp
@@ -1264,6 +1264,16 @@ namespace UE::RivermaxCore::Private
 		if(Status != RMX_OK)
 		{
 			OutErrorMessage = FString::Printf(TEXT("Could not attach flow to stream. Status: %d."), Status);
+
+			// Cleanup stream.
+			Status = CachedAPI->rmx_input_destroy_stream(StreamId);
+
+			if (Status != RMX_OK)
+			{
+				UE_LOG(LogRivermax, Warning, TEXT("Failed to destroy input stream %d correctly. Status: %d"), StreamId, Status);
+			}
+
+			return false;
 		}
 		
 		return true;
