Fix Rivermax YUV10Bit422ToRGBACS shader function

diff --git a/Engine/Plugins/VirtualProduction/Rivermax/RivermaxCore/Shaders/Private/RivermaxShaders.usf b/Engine/Plugins/VirtualProduction/Rivermax/RivermaxCore/Shaders/Private/RivermaxShaders.usf
--- a/Engine/Plugins/VirtualProduction/Rivermax/RivermaxCore/Shaders/Private/RivermaxShaders.usf
+++ b/Engine/Plugins/VirtualProduction/Rivermax/RivermaxCore/Shaders/Private/RivermaxShaders.usf
@@ -749,7 +749,7 @@ void YUV10Bit422ToRGBACS(uint3 ThreadId : SV_DispatchThreadID, uint GroupIndex :
 	PixelData[2].YUVY.y = ((((InputYCbCrBuffer[BufferIndex].DWord2 >> 16)  << 2) | ((InputYCbCrBuffer[BufferIndex].DWord2 >> 30) & 0x003)) & 0x3FF) * Mul;
 	PixelData[2].YUVY.z = ((((InputYCbCrBuffer[BufferIndex].DWord3 >> 00)  << 6) | ((InputYCbCrBuffer[BufferIndex].DWord3 >> 10) & 0x03F)) & 0x3FF) * Mul;
 	PixelData[2].YUVY.w = ((((InputYCbCrBuffer[BufferIndex].DWord3 >>  8)  << 8) | ((InputYCbCrBuffer[BufferIndex].DWord3 >> 16) & 0x0FF)) & 0x3FF) * Mul;
-	PixelData[3].YUVY.x = ((((InputYCbCrBuffer[BufferIndex].DWord4 >>  0)  << 4) | ((InputYCbCrBuffer[BufferIndex].DWord3 >> 12) & 0x00F)) & 0x3FF) * Mul;
+	PixelData[3].YUVY.x = ((((InputYCbCrBuffer[BufferIndex].DWord4 >>  0)  << 4) | ((InputYCbCrBuffer[BufferIndex].DWord4 >> 12) & 0x00F)) & 0x3FF) * Mul;
 	PixelData[3].YUVY.y = ((((InputYCbCrBuffer[BufferIndex].DWord3 >> 24)  << 2) | ((InputYCbCrBuffer[BufferIndex].DWord4 >> 06) & 0x003)) & 0x3FF) * Mul;
 	PixelData[3].YUVY.z = ((((InputYCbCrBuffer[BufferIndex].DWord4 >>  8)  << 6) | ((InputYCbCrBuffer[BufferIndex].DWord4 >> 18) & 0x03F)) & 0x3FF) * Mul;
 	PixelData[3].YUVY.w = ((((InputYCbCrBuffer[BufferIndex].DWord4 >> 16)  << 8) | ((InputYCbCrBuffer[BufferIndex].DWord4 >> 24) & 0x0FF)) & 0x3FF) * Mul;
