Rivermax: Made requirements for SDK environment path less rigid.

diff --git a/Engine/Plugins/VirtualProduction/Rivermax/RivermaxCore/Source/RivermaxCore/Private/RivermaxManager.cpp b/Engine/Plugins/VirtualProduction/Rivermax/RivermaxCore/Source/RivermaxCore/Private/RivermaxManager.cpp
--- a/Engine/Plugins/VirtualProduction/Rivermax/RivermaxCore/Source/RivermaxCore/Private/RivermaxManager.cpp
+++ b/Engine/Plugins/VirtualProduction/Rivermax/RivermaxCore/Source/RivermaxCore/Private/RivermaxManager.cpp
@@ -41,7 +41,7 @@ namespace UE::RivermaxCore
 	TMap<uint8, FString> FRivermaxTracingUtils::RmaxInSelectedFrameTraceEvents;
 
 	// Environment variable used to specify where Rivermax library is located. This enables having an installed version different than the one used by Unreal
-	static const FString RivermaxLibraryEnvironmentVariable = TEXT("RIVERMAX_PATH_1_3x");
+	static const FString RivermaxLibraryEnvironmentVariable = TEXT("RIVERMAX_PATH");
 }
 
 namespace UE::RivermaxCore::Private
