Rivermax: Improvements.
- Commit chunks before the next alignment point.
- fallback to more than one srd size only when it is absolutely necessary and the row cannot be divided onto even SRDs
- Added more traces for debugging purposes.

diff --git a/Engine/Plugins/VirtualProduction/Rivermax/RivermaxCore/Source/RivermaxCore/Private/RivermaxOutputStream.cpp b/Engine/Plugins/VirtualProduction/Rivermax/RivermaxCore/Source/RivermaxCore/Private/RivermaxOutputStream.cpp
--- a/Engine/Plugins/VirtualProduction/Rivermax/RivermaxCore/Source/RivermaxCore/Private/RivermaxOutputStream.cpp
+++ b/Engine/Plugins/VirtualProduction/Rivermax/RivermaxCore/Source/RivermaxCore/Private/RivermaxOutputStream.cpp
@@ -41,8 +41,8 @@ namespace UE::RivermaxCore::Private
 
 	static TAutoConsoleVariable<int32> CVarRivermaxOutputEnableMultiSRD(
 		TEXT("Rivermax.Output.EnableMultiSRD"), 1,
-		TEXT("When enabled, payload size will be constant across the frame except the last one.\n" 
-		     "If disabled, a payload size that fits the line will be used causing some resolution to not be supported."),
+		TEXT("When enabled and if the row cannot be split evenly, non-uniform payloads will be used. The last packet for the frame will not be fully filled with data.\n" 
+		     "If disabled, the payloads will be split evenly or the 2110 stream will be disabled."),
 		ECVF_Default);
 
 	static TAutoConsoleVariable<int32> CVarRivermaxOutputLinesPerChunk(
@@ -123,6 +123,12 @@ namespace UE::RivermaxCore::Private
 		TEXT("Multiplier applied to desired output frame rate in order to reduce time it takes to send out a frame and slowly correct misalignment that could happen."),
 		ECVF_Default);
 
+	static TAutoConsoleVariable<float> CVarRivermaxOutputCommitChunksOffsetPercent(
+		TEXT("Rivermax.Output.CommitChunksOffsetPercent"), 0.3,
+		TEXT("This CVar will allow for chunks to be committed before the next alignment point if it is at all possible.\n\
+		The value indicates in percent of frame time how much earlier the Rivermax plugin will attempt to commit chunks."),
+		ECVF_Default);
+
 	static bool GbTriggerRandomTimingIssue = false;
 	FAutoConsoleVariableRef CVarTriggerRandomTimingIssue(
 		TEXT("Rivermax.Sync.TriggerRandomTimingIssue")
@@ -772,11 +778,21 @@ namespace UE::RivermaxCore::Private
 				if (PostCopyTimeLeftSec > 0)
 				{
 					TRACE_CPUPROFILER_EVENT_SCOPE(RmaxOut::Waiting);
-					static constexpr float YieldTimeSec = 2.0f / 1000;
+					static constexpr float YieldTimeSec = 1.0f / 1000;
+
+					double TargetPlatformTimeToCommitSec = TargetPlatformTimeSec;
 
-					while (FPlatformTime::Seconds() < TargetPlatformTimeSec && bIsActive)
+					const double CVarStartCommitEarlyPercent = FMath::Clamp(CVarRivermaxOutputCommitChunksOffsetPercent.GetValueOnAnyThread(), 0.0, 0.8);
+					if (CVarStartCommitEarlyPercent > 0.)
 					{
-						const double TimeLeft = TargetPlatformTimeSec - FPlatformTime::Seconds();
+						// We would like to start comitting chunks before the alignment point, but not too early so that we don't block this thread.
+						// This will force the chunks to be committed CVarRivermaxOutputMaxFrameMemorySliceCount % of the frame time earlier than the next alignment point
+						TargetPlatformTimeToCommitSec = TargetPlatformTimeSec - Options.FrameRate.AsInterval() * CVarStartCommitEarlyPercent;
+					}
+
+					while (FPlatformTime::Seconds() < TargetPlatformTimeToCommitSec && bIsActive)
+					{
+						const double TimeLeft = TargetPlatformTimeToCommitSec - FPlatformTime::Seconds();
 						const double SleepTime = TimeLeft > YieldTimeSec ? TimeLeft - YieldTimeSec : 0.0;
 						FPlatformProcess::SleepNoStats(SleepTime);
 					}
@@ -813,31 +829,32 @@ namespace UE::RivermaxCore::Private
 
 		const int32 BytesPerLine = GetStride();
 
-		// Find out payload we want to use. Either we go the 'potential' multi SRD route or we keep the old way of finding a common payload
-		// with more restrictions on resolution supported. Kept in place to be able to fallback in case there are issues with the multiSRD one.
-		if (CVarRivermaxOutputEnableMultiSRD.GetValueOnAnyThread() >= 1)
+		// By default we want to divide the bytes evenly across packets. Some resolutions will require packets to be sized unevenly.
+		const bool bFoundPayload = FindPayloadSize(Options, BytesPerLine, FormatInfo, StreamMemory.PayloadSize);
+		if (bFoundPayload == false)
 		{
-			if (CVarRivermaxOutputMaximizePacketSize.GetValueOnAnyThread() >= 1)
+			// Find out payload we want to use. Either we go the 'potential' multi SRD route or we restrict the stream based on supported resolutions.
+			if (CVarRivermaxOutputEnableMultiSRD.GetValueOnAnyThread() >= 1)
 			{
-				StreamMemory.PayloadSize = GetMaximizedPayloadSize(FormatInfo.Sampling);
+				UE_LOG(LogRivermax, Log, TEXT("Due to resolution %dx%d, row data will be sent over multiple packets with varied sizes."), Options.AlignedResolution.X, Options.AlignedResolution.Y);
+				if (CVarRivermaxOutputMaximizePacketSize.GetValueOnAnyThread() >= 1)
+				{
+					StreamMemory.PayloadSize = GetMaximizedPayloadSize(FormatInfo.Sampling);
+				}
+				else
+				{
+					StreamMemory.PayloadSize = GetPayloadSize(FormatInfo.Sampling);
+				}
 			}
 			else
 			{
-				StreamMemory.PayloadSize = GetPayloadSize(FormatInfo.Sampling);
-			}
-		}
-		else
-		{
-			const bool bFoundPayload = FindPayloadSize(Options, BytesPerLine, FormatInfo, StreamMemory.PayloadSize);
-			if (bFoundPayload == false)
-			{
-				UE_LOG(LogRivermax, Warning, TEXT("Could not find payload size for desired resolution %dx%d for desired pixel format"), Options.AlignedResolution.X, Options.AlignedResolution.Y);
+				UE_LOG(LogRivermax, Warning, TEXT("Could not find payload size for desired resolution %dx%d for desired pixel format."
+					"If the intention is to use non standard resolutions, users might want to enable multi-srd support via Rivermax.Output.EnableMultiSRD."), Options.AlignedResolution.X, Options.AlignedResolution.Y);
 				return false;
 			}
 		}
 
 		// With payload size in hand, figure out how many packets we will need, how many chunks (group of packets) and configure descriptor arrays
-
 		const uint32 PixelCount = Options.AlignedResolution.X * Options.AlignedResolution.Y;
 		const uint64 FrameSize = PixelCount / FormatInfo.PixelGroupCoverage * FormatInfo.PixelGroupSize;
 
@@ -935,12 +952,26 @@ namespace UE::RivermaxCore::Private
 	
 		uint64 TotalSize = 0;
 		uint64 LineSize = 0;
-		for (int32 PayloadSizeIndex = 0; PayloadSizeIndex < RealPacketsPerFrame; ++PayloadSizeIndex)
+		for (uint32 PayloadSizeIndex = 0; PayloadSizeIndex < RealPacketsPerFrame; ++PayloadSizeIndex)
 		{
 			uint32 HeaderSize = FRawRTPHeader::OneSRDSize;
 			uint32 ThisPayloadSize = StreamMemory.PayloadSize;
+
+			// The last truly valid packet is smaller in size.
+			if (PayloadSizeIndex == StreamMemory.PacketsPerFrame - 1)
+			{
+				ThisPayloadSize = FrameSize - (StreamMemory.PacketsPerFrame - 1)* StreamMemory.PayloadSize;
+			}
+			else if (PayloadSizeIndex >= StreamMemory.PacketsPerFrame)
+			{
+				// Extra header/payload required for the chunk alignment are set to 0. Nothing has to be sent out the wire.
+				HeaderSize = 0;
+				ThisPayloadSize = 0;
+			}
+
 			if (TotalSize < FrameSize)
 			{
+
 				if ((LineSize + StreamMemory.PayloadSize) == BytesPerLine)
 				{
 					LineSize = 0;
@@ -966,12 +997,6 @@ namespace UE::RivermaxCore::Private
 					HeaderSize = FRawRTPHeader::OneSRDSize;
 				}
 			}
-			else
-			{
-				// Extra header/payload required for the chunk alignment are set to 0. Nothing has to be sent out the wire.
-				HeaderSize = 0;
-				ThisPayloadSize = 0;
-			}
 
 			// All buffers are configured the same so compute header and payload sizes once and assigned to all impacted locations
 			for (uint32 BufferIndex = 0; BufferIndex < StreamMemory.FramesFieldPerMemoryBlock; ++BufferIndex)
@@ -985,7 +1010,7 @@ namespace UE::RivermaxCore::Private
 				RTPFiller.Update(PayloadSizeIndex);
 			}
 			
-			TotalSize += ThisPayloadSize;
+			TotalSize += StreamMemory.PayloadSize;
 		}
 
 		// Verify memcopy config to make sure it works for current frame size / chunking
@@ -1121,6 +1146,7 @@ namespace UE::RivermaxCore::Private
 
 	bool FRivermaxOutputStream::WaitForNextRound()
 	{
+		TRACE_CPUPROFILER_EVENT_SCOPE(WaitForNextRound);
 		const uint64 CurrentTimeNanosec = RivermaxModule->GetRivermaxManager()->GetTime();
 		const double CurrentPlatformTime = FPlatformTime::Seconds();
 		const uint64 CurrentFrameNumber = UE::RivermaxCore::GetFrameNumber(CurrentTimeNanosec, Options.FrameRate);
@@ -1235,6 +1261,8 @@ namespace UE::RivermaxCore::Private
 
 	void FRivermaxOutputStream::GetNextChunk()
 	{
+		TRACE_CPUPROFILER_EVENT_SCOPE(GetNextChunk);
+
 		bool bHasAddedTrace = false;
 		rmx_status Status;
 		do
@@ -1355,6 +1383,7 @@ namespace UE::RivermaxCore::Private
 
 	void FRivermaxOutputStream::CommitNextChunks()
 	{
+		TRACE_CPUPROFILER_EVENT_SCOPE(CommitNextChunks);
 		rmx_status Status;
 		int32 ErrorCount = 0;
 		const uint64 CurrentTimeNanosec = RivermaxModule->GetRivermaxManager()->GetTime();
@@ -1455,6 +1484,8 @@ namespace UE::RivermaxCore::Private
 
 	void FRivermaxOutputStream::PrepareNextFrame()
 	{
+		TRACE_CPUPROFILER_EVENT_SCOPE(PrepareNextFrame);
+
 		switch (Options.AlignmentMode)
 		{
 			case ERivermaxAlignmentMode::FrameCreation:
@@ -2026,6 +2057,7 @@ namespace UE::RivermaxCore::Private
 
 	void FRivermaxOutputStream::SendFrame()
 	{
+		TRACE_CPUPROFILER_EVENT_SCOPE(SendFrame);
 		StreamData.LastSendStartTimeNanoSec = RivermaxModule->GetRivermaxManager()->GetTime();
 
 		if (bTriggerRandomDelay)
diff --git a/Engine/Plugins/VirtualProduction/Rivermax/RivermaxCore/Source/RivermaxCore/Private/RivermaxUtils.h b/Engine/Plugins/VirtualProduction/Rivermax/RivermaxCore/Source/RivermaxCore/Private/RivermaxUtils.h
--- a/Engine/Plugins/VirtualProduction/Rivermax/RivermaxCore/Source/RivermaxCore/Private/RivermaxUtils.h
+++ b/Engine/Plugins/VirtualProduction/Rivermax/RivermaxCore/Source/RivermaxCore/Private/RivermaxUtils.h
@@ -18,11 +18,11 @@ namespace UE::RivermaxCore::Private::Utils
 	static constexpr uint32 FullHDHeight = 1080;
 	static constexpr uint32 FullHDWidth = 1920;
 
-	/** Maximum payload we can send based on UDP max size and RTP header.  */
+	/** Maximum payload size in bytes that the plugin can send based on UDP max size and RTP header.  */
 	static constexpr uint32 MaxPayloadSize = 1420;
 	
-	/** Smallest payload size to bound our search for a payload that can be equal across a line */
-	static constexpr uint32 MinPayloadSize = 950;
+	/** Smallest payload size (bytes) to use as a lower bound in search for a payload that can be equal across a line */
+	static constexpr uint32 MinPayloadSize = 600;
 
 	/** SMTPE 2110-10.The Media Clock and RTP Clock rate for streams compliant to this standard shall be 90 kHz. */
 	static constexpr double MediaClockSampleRate = 90000.0;
