Rivermax: Since PCI transactions happen per packet and not per chunk, and multi-SRD issue is fixed, it makes sense to support Mutli-SRD first and even split second.

diff --git a/Engine/Plugins/VirtualProduction/Rivermax/RivermaxCore/Source/RivermaxCore/Private/RivermaxOutputStream.cpp b/Engine/Plugins/VirtualProduction/Rivermax/RivermaxCore/Source/RivermaxCore/Private/RivermaxOutputStream.cpp
--- a/Engine/Plugins/VirtualProduction/Rivermax/RivermaxCore/Source/RivermaxCore/Private/RivermaxOutputStream.cpp
+++ b/Engine/Plugins/VirtualProduction/Rivermax/RivermaxCore/Source/RivermaxCore/Private/RivermaxOutputStream.cpp
@@ -41,7 +41,7 @@ namespace UE::RivermaxCore::Private
 
 	static TAutoConsoleVariable<int32> CVarRivermaxOutputEnableMultiSRD(
 		TEXT("Rivermax.Output.EnableMultiSRD"), 1,
-		TEXT("When enabled and if the row cannot be split evenly, non-uniform payloads will be used. The last packet for the frame will not be fully filled with data.\n" 
+		TEXT("When enabled, non-uniform payloads will be used. The last packet for the frame will not be fully filled with data.\n" 
 		     "If disabled, the payloads will be split evenly or the 2110 stream will be disabled."),
 		ECVF_Default);
 
@@ -829,24 +829,23 @@ namespace UE::RivermaxCore::Private
 
 		const int32 BytesPerLine = GetStride();
 
-		// By default we want to divide the bytes evenly across packets. Some resolutions will require packets to be sized unevenly.
-		const bool bFoundPayload = FindPayloadSize(Options, BytesPerLine, FormatInfo, StreamMemory.PayloadSize);
-		if (bFoundPayload == false)
+		// Find out payload we want to use. Either we go the 'potential' multi SRD route or we keep the old way of finding a common payload
+		// with more restrictions on resolution supported. Kept in place to be able to fallback in case there are issues with the multiSRD one.
+		if (CVarRivermaxOutputEnableMultiSRD.GetValueOnAnyThread() >= 1)
 		{
-			// Find out payload we want to use. Either we go the 'potential' multi SRD route or we restrict the stream based on supported resolutions.
-			if (CVarRivermaxOutputEnableMultiSRD.GetValueOnAnyThread() >= 1)
+			if (CVarRivermaxOutputMaximizePacketSize.GetValueOnAnyThread() >= 1)
 			{
-				UE_LOG(LogRivermax, Log, TEXT("Due to resolution %dx%d, row data will be sent over multiple packets with varied sizes."), Options.AlignedResolution.X, Options.AlignedResolution.Y);
-				if (CVarRivermaxOutputMaximizePacketSize.GetValueOnAnyThread() >= 1)
-				{
-					StreamMemory.PayloadSize = GetMaximizedPayloadSize(FormatInfo.Sampling);
-				}
-				else
-				{
-					StreamMemory.PayloadSize = GetPayloadSize(FormatInfo.Sampling);
-				}
+				StreamMemory.PayloadSize = GetMaximizedPayloadSize(FormatInfo.Sampling);
 			}
 			else
+			{
+				StreamMemory.PayloadSize = GetPayloadSize(FormatInfo.Sampling);
+			}
+		}
+		else
+		{
+			const bool bFoundPayload = FindPayloadSize(Options, BytesPerLine, FormatInfo, StreamMemory.PayloadSize);
+			if (bFoundPayload == false)
 			{
 				UE_LOG(LogRivermax, Warning, TEXT("Could not find payload size for desired resolution %dx%d for desired pixel format."
 					"If the intention is to use non standard resolutions, users might want to enable multi-srd support via Rivermax.Output.EnableMultiSRD."), Options.AlignedResolution.X, Options.AlignedResolution.Y);
