Undo CL: Rivermax: Since PCI transactions happen per packet and not per chunk, and multi-SRD issue is fixed, it makes sense to support Mutli-SRD first and even split second.

diff --git a/Engine/Plugins/VirtualProduction/Rivermax/RivermaxCore/Source/RivermaxCore/Private/RivermaxOutputStream.cpp b/Engine/Plugins/VirtualProduction/Rivermax/RivermaxCore/Source/RivermaxCore/Private/RivermaxOutputStream.cpp
--- a/Engine/Plugins/VirtualProduction/Rivermax/RivermaxCore/Source/RivermaxCore/Private/RivermaxOutputStream.cpp
+++ b/Engine/Plugins/VirtualProduction/Rivermax/RivermaxCore/Source/RivermaxCore/Private/RivermaxOutputStream.cpp
@@ -41,7 +41,7 @@ namespace UE::RivermaxCore::Private
 
 	static TAutoConsoleVariable<int32> CVarRivermaxOutputEnableMultiSRD(
 		TEXT("Rivermax.Output.EnableMultiSRD"), 1,
-		TEXT("When enabled, non-uniform payloads will be used. The last packet for the frame will not be fully filled with data.\n" 
+		TEXT("When enabled and if the row cannot be split evenly, non-uniform payloads will be used. The last packet for the frame will not be fully filled with data.\n" 
 		     "If disabled, the payloads will be split evenly or the 2110 stream will be disabled."),
 		ECVF_Default);
 
@@ -829,23 +829,24 @@ namespace UE::RivermaxCore::Private
 
 		const int32 BytesPerLine = GetStride();
 
-		// Find out payload we want to use. Either we go the 'potential' multi SRD route or we keep the old way of finding a common payload
-		// with more restrictions on resolution supported. Kept in place to be able to fallback in case there are issues with the multiSRD one.
-		if (CVarRivermaxOutputEnableMultiSRD.GetValueOnAnyThread() >= 1)
+		// By default we want to divide the bytes evenly across packets. Some resolutions will require packets to be sized unevenly.
+		const bool bFoundPayload = FindPayloadSize(Options, BytesPerLine, FormatInfo, StreamMemory.PayloadSize);
+		if (bFoundPayload == false)
 		{
-			if (CVarRivermaxOutputMaximizePacketSize.GetValueOnAnyThread() >= 1)
+			// Find out payload we want to use. Either we go the 'potential' multi SRD route or we restrict the stream based on supported resolutions.
+			if (CVarRivermaxOutputEnableMultiSRD.GetValueOnAnyThread() >= 1)
 			{
-				StreamMemory.PayloadSize = GetMaximizedPayloadSize(FormatInfo.Sampling);
+				UE_LOG(LogRivermax, Log, TEXT("Due to resolution %dx%d, row data will be sent over multiple packets with varied sizes."), Options.AlignedResolution.X, Options.AlignedResolution.Y);
+				if (CVarRivermaxOutputMaximizePacketSize.GetValueOnAnyThread() >= 1)
+				{
+					StreamMemory.PayloadSize = GetMaximizedPayloadSize(FormatInfo.Sampling);
+				}
+				else
+				{
+					StreamMemory.PayloadSize = GetPayloadSize(FormatInfo.Sampling);
+				}
 			}
 			else
-			{
-				StreamMemory.PayloadSize = GetPayloadSize(FormatInfo.Sampling);
-			}
-		}
-		else
-		{
-			const bool bFoundPayload = FindPayloadSize(Options, BytesPerLine, FormatInfo, StreamMemory.PayloadSize);
-			if (bFoundPayload == false)
 			{
 				UE_LOG(LogRivermax, Warning, TEXT("Could not find payload size for desired resolution %dx%d for desired pixel format."
 					"If the intention is to use non standard resolutions, users might want to enable multi-srd support via Rivermax.Output.EnableMultiSRD."), Options.AlignedResolution.X, Options.AlignedResolution.Y);
@@ -951,26 +952,12 @@ namespace UE::RivermaxCore::Private
 	
 		uint64 TotalSize = 0;
 		uint64 LineSize = 0;
-		for (uint32 PayloadSizeIndex = 0; PayloadSizeIndex < RealPacketsPerFrame; ++PayloadSizeIndex)
+		for (int32 PayloadSizeIndex = 0; PayloadSizeIndex < RealPacketsPerFrame; ++PayloadSizeIndex)
 		{
 			uint32 HeaderSize = FRawRTPHeader::OneSRDSize;
 			uint32 ThisPayloadSize = StreamMemory.PayloadSize;
-
-			// The last truly valid packet is smaller in size.
-			if (PayloadSizeIndex == StreamMemory.PacketsPerFrame - 1)
-			{
-				ThisPayloadSize = FrameSize - (StreamMemory.PacketsPerFrame - 1)* StreamMemory.PayloadSize;
-			}
-			else if (PayloadSizeIndex >= StreamMemory.PacketsPerFrame)
-			{
-				// Extra header/payload required for the chunk alignment are set to 0. Nothing has to be sent out the wire.
-				HeaderSize = 0;
-				ThisPayloadSize = 0;
-			}
-
 			if (TotalSize < FrameSize)
 			{
-
 				if ((LineSize + StreamMemory.PayloadSize) == BytesPerLine)
 				{
 					LineSize = 0;
@@ -996,6 +983,12 @@ namespace UE::RivermaxCore::Private
 					HeaderSize = FRawRTPHeader::OneSRDSize;
 				}
 			}
+			else
+			{
+				// Extra header/payload required for the chunk alignment are set to 0. Nothing has to be sent out the wire.
+				HeaderSize = 0;
+				ThisPayloadSize = 0;
+			}
 
 			// All buffers are configured the same so compute header and payload sizes once and assigned to all impacted locations
 			for (uint32 BufferIndex = 0; BufferIndex < StreamMemory.FramesFieldPerMemoryBlock; ++BufferIndex)
@@ -1009,7 +1002,7 @@ namespace UE::RivermaxCore::Private
 				RTPFiller.Update(PayloadSizeIndex);
 			}
 			
-			TotalSize += StreamMemory.PayloadSize;
+			TotalSize += ThisPayloadSize;
 		}
 
 		// Verify memcopy config to make sure it works for current frame size / chunking
