Rivermax Media Output: if intermediate buffer is used - we copy data as we commit it. This fixes occasional tearing.

#rnx

diff --git a/Engine/Plugins/VirtualProduction/Rivermax/RivermaxCore/Source/RivermaxCore/Private/RivermaxOutputStream.cpp b/Engine/Plugins/VirtualProduction/Rivermax/RivermaxCore/Source/RivermaxCore/Private/RivermaxOutputStream.cpp
--- a/Engine/Plugins/VirtualProduction/Rivermax/RivermaxCore/Source/RivermaxCore/Private/RivermaxOutputStream.cpp
+++ b/Engine/Plugins/VirtualProduction/Rivermax/RivermaxCore/Source/RivermaxCore/Private/RivermaxOutputStream.cpp
@@ -740,7 +740,6 @@ namespace UE::RivermaxCore::Private
 			}
 
 			Stats.TotalChunkRetries += Stats.LastFrameChunkRetries;
-			StreamData.bHasFrameFirstChunkBeenFetched = false;
 			Stats.LastFrameChunkRetries = 0;
 		}
 	}
@@ -763,6 +762,7 @@ namespace UE::RivermaxCore::Private
 			const double TargetPlatformTimeSec = CurrentPlatformTime + TimeLeftSec;
 			if (CurrentRmaxTimeNanosec < StreamData.NextAlignmentPointNanosec)
 			{
+				if (!StreamMemory.bUseIntermediateBuffer)
 				{
 					TRACE_CPUPROFILER_EVENT_SCOPE(RmaxOut::CopyFrame);
 
@@ -1264,10 +1264,8 @@ namespace UE::RivermaxCore::Private
 			CurrentFrame->HeaderPtr = rmx_output_media_get_chunk_strides(&StreamData.ChunkHandle, StreamMemory.HeaderBlockID);
 			if (Status == RMX_OK)
 			{
-				if (StreamData.bHasFrameFirstChunkBeenFetched == false)
+				if (CurrentFrame->FrameStartPtr == nullptr)
 				{
-					StreamData.bHasFrameFirstChunkBeenFetched = true;
-
 					// Stamp frame start in order to copy frame data sequentially as we query chunks
 					CurrentFrame->FrameStartPtr = CurrentFrame->PayloadPtr;
 
diff --git a/Engine/Plugins/VirtualProduction/Rivermax/RivermaxCore/Source/RivermaxCore/Private/RivermaxOutputStream.h b/Engine/Plugins/VirtualProduction/Rivermax/RivermaxCore/Source/RivermaxCore/Private/RivermaxOutputStream.h
--- a/Engine/Plugins/VirtualProduction/Rivermax/RivermaxCore/Source/RivermaxCore/Private/RivermaxOutputStream.h
+++ b/Engine/Plugins/VirtualProduction/Rivermax/RivermaxCore/Source/RivermaxCore/Private/RivermaxOutputStream.h
@@ -154,9 +154,6 @@ namespace UE::RivermaxCore::Private
 		/** Data and RTP frame index expected to be used for next frame */
 		uint8 ExpectedFrameIndex = 0;
 
-		/** Used to detect misalignment between chunk being sent and frame memory we are writing in */
-		bool bHasFrameFirstChunkBeenFetched = false;
-
 		/** Next alignment point based on PTP standard */
 		uint64 NextAlignmentPointNanosec = 0;
 
