Rivermax: Cleanup:
- Removed unused functionality and Cvars
- Unified RTP header creation code
- properly aligned RTP headers so that it returns the right size via sizeof
- all pointers stored and pointer operations are now using 64 bit data type.
- Frame management is now done on shared pointers instead of integers. Reduced the number of managed arrays.
- removed unused members and pieces of code.
- When multi-srd is used the last packet is now estimated correctly.

[RN]

diff --git a/Engine/Plugins/VirtualProduction/Rivermax/RivermaxCore/Source/RivermaxCore/Private/RTPHeader.cpp b/Engine/Plugins/VirtualProduction/Rivermax/RivermaxCore/Source/RivermaxCore/Private/RTPHeader.cpp
--- a/Engine/Plugins/VirtualProduction/Rivermax/RivermaxCore/Source/RivermaxCore/Private/RTPHeader.cpp
+++ b/Engine/Plugins/VirtualProduction/Rivermax/RivermaxCore/Source/RivermaxCore/Private/RTPHeader.cpp
@@ -6,48 +6,26 @@
 
 namespace UE::RivermaxCore::Private
 {
-	uint16 FRawRTPHeader::GetSrd1RowNumber() const
+	uint16 FRawSRD::GetRowNumber() const
 	{
-		return ((SRD1RowNumberHigh << 8) | SRD1RowNumberLow);
+		return ((RowNumberHigh << 8) | RowNumberLow);
 	}
 
-	void FRawRTPHeader::SetSrd1RowNumber(uint16 RowNumber)
+	void FRawSRD::SetRowNumber(uint16 RowNumber)
 	{
-		SRD1RowNumberHigh = (RowNumber >> 8) & 0xFF;
-		SRD1RowNumberLow = RowNumber & 0xFF;
+		RowNumberHigh = (RowNumber >> 8) & 0xFF;
+		RowNumberLow = RowNumber & 0xFF;
 	}
 
-	uint16 FRawRTPHeader::GetSrd1Offset() const
+	uint16 FRawSRD::GetOffset() const
 	{
-		return ((SRD1OffsetHigh << 8) | SRD1OffsetLow);
+		return ((OffsetHigh << 8) | OffsetLow);
 	}
 
-	void FRawRTPHeader::SetSrd1Offset(uint16 Offset)
+	void FRawSRD::SetOffset(uint16 Offset)
 	{
-		SRD1OffsetHigh = (Offset >> 8) & 0xFF;
-		SRD1OffsetLow = Offset & 0xFF;
-	}
-
-	uint16 FRawRTPHeader::GetSrd2RowNumber() const
-	{
-		return ((SRD2RowNumberHigh << 8) | SRD2RowNumberLow);
-	}
-
-	void FRawRTPHeader::SetSrd2RowNumber(uint16 RowNumber)
-	{
-		SRD2RowNumberHigh = (RowNumber >> 8) & 0xFF;
-		SRD2RowNumberLow = RowNumber & 0xFF;
-	}
-
-	uint16 FRawRTPHeader::GetSrd2Offset() const
-	{
-		return ((SRD2OffsetHigh << 8) | SRD2OffsetLow);
-	}
-
-	void FRawRTPHeader::SetSrd2Offset(uint16 Offset)
-	{
-		SRD2OffsetHigh = (Offset >> 8) & 0xFF;
-		SRD2OffsetLow = Offset & 0xFF;
+		OffsetHigh = (Offset >> 8) & 0xFF;
+		OffsetLow = Offset & 0xFF;
 	}
 
 	const uint8* GetRTPHeaderPointer(const uint8* InHeader)
@@ -84,19 +62,19 @@ namespace UE::RivermaxCore::Private
 
 		SyncSouceId = RawHeader.SynchronizationSource;
 
-		SRD1.Length = ByteSwap((uint16)RawHeader.SRD1Length);
-		SRD1.DataOffset = RawHeader.GetSrd1Offset();
-		SRD1.RowNumber = RawHeader.GetSrd1RowNumber();
-		SRD1.bIsFieldOne = RawHeader.FieldIdentification1;
-		SRD1.bHasContinuation = RawHeader.ContinuationBit1;
+		SRD1.Length = ByteSwap((uint16)RawHeader.SRD1.Length);
+		SRD1.DataOffset = RawHeader.SRD1.GetOffset();
+		SRD1.RowNumber = RawHeader.SRD1.GetRowNumber();
+		SRD1.bIsFieldOne = RawHeader.SRD1.FieldIdentification;
+		SRD1.bHasContinuation = RawHeader.SRD1.ContinuationBit;
 
 		if (SRD1.bHasContinuation)
 		{
-			SRD2.Length = ByteSwap((uint16)RawHeader.SRD2Length);
-			SRD2.DataOffset = RawHeader.GetSrd2Offset();
-			SRD2.RowNumber = RawHeader.GetSrd2RowNumber();
-			SRD2.bIsFieldOne = RawHeader.FieldIdentification2;
-			SRD2.bHasContinuation = RawHeader.ContinuationBit2;
+			SRD2.Length = ByteSwap((uint16)RawHeader.SRD2.Length);
+			SRD2.DataOffset = RawHeader.SRD2.GetOffset();
+			SRD2.RowNumber = RawHeader.SRD2.GetRowNumber();
+			SRD2.bIsFieldOne = RawHeader.SRD2.FieldIdentification;
+			SRD2.bHasContinuation = RawHeader.SRD2.ContinuationBit;
 
 			if (SRD2.bHasContinuation == true)
 			{
diff --git a/Engine/Plugins/VirtualProduction/Rivermax/RivermaxCore/Source/RivermaxCore/Private/RTPHeader.h b/Engine/Plugins/VirtualProduction/Rivermax/RivermaxCore/Source/RivermaxCore/Private/RTPHeader.h
--- a/Engine/Plugins/VirtualProduction/Rivermax/RivermaxCore/Source/RivermaxCore/Private/RTPHeader.h
+++ b/Engine/Plugins/VirtualProduction/Rivermax/RivermaxCore/Source/RivermaxCore/Private/RTPHeader.h
@@ -36,58 +36,54 @@ namespace UE::RivermaxCore::Private
 #if PLATFORM_SUPPORTS_PRAGMA_PACK
 #pragma pack(push, 1)
 #endif
+	/** Total packed size should be 6 octets. */
+	struct FRawSRD
+	{
+		uint16 Length				: 16;
+
+		uint8 RowNumberHigh			: 7;
+
+		uint8 FieldIdentification	: 1;
+
+		uint8 RowNumberLow			: 8;
+
+		uint8 OffsetHigh			: 7;
+
+		/** If set indicates that there is another SRD following this one.*/
+		uint8 ContinuationBit		: 1;
+		uint8 OffsetLow				: 8;
+
+
+		/** Returns SRD associated row number */
+		uint16 GetRowNumber() const;
+
+		/** Sets SRD associated row number */
+		void SetRowNumber(uint16 RowNumber);
+
+		/** Returns SRD pixel offset in its associated row */
+		uint16 GetOffset() const;
+
+		/** Sets SRD pixel offset in its associated row */
+		void SetOffset(uint16 Offset);
+	};
+
+	/** Total packed size should be 26 octets. */
 	struct FRawRTPHeader
 	{
-		uint32 ContributingSourceCount : 4;
-		uint32 ExtensionBit : 1;
-		uint32 PaddingBit : 1;
-		uint32 Version : 2;
-		uint32 PayloadType : 7;
-		uint32 MarkerBit : 1;
-		uint32 SequenceNumber : 16;
-		uint32 Timestamp : 32;
-		uint32 SynchronizationSource : 32;
-		uint32 ExtendedSequenceNumber : 16;
-		//SRD 1
-		uint32 SRD1Length : 16;
-		uint32 SRD1RowNumberHigh : 7;
-		uint32 FieldIdentification1 : 1;
-		uint32 SRD1RowNumberLow : 8;
-		uint32 SRD1OffsetHigh : 7;
-		uint32 ContinuationBit1 : 1;
-		uint32 SRD1OffsetLow : 8;
-		//SRD 2
-		uint32 SRD2Length : 16;
-		uint32 SRD2RowNumberHigh : 7;
-		uint32 FieldIdentification2 : 1;
-		uint32 SRD2RowNumberLow : 8;
-		uint32 SRD2OffsetHigh : 7;
-		uint32 ContinuationBit2 : 1;
-		uint32 SRD2OffsetLow : 8;
-
-		/** Returns SRD1 associated row number */
-		uint16 GetSrd1RowNumber() const;
-
-		/** Sets SRD1 associated row number */
-		void SetSrd1RowNumber(uint16 RowNumber);
-
-		/** Returns SRD1 pixel offset in its associated row */
-		uint16 GetSrd1Offset() const;
-
-		/** Sets SRD1 pixel offset in its associated row */
-		void SetSrd1Offset(uint16 Offset);
-
-		/** Sets SRD2 associated row number */
-		uint16 GetSrd2RowNumber() const;
-
-		/** Sets SRD1 associated row number */
-		void SetSrd2RowNumber(uint16 RowNumber);
-
-		/** Returns SRD1 pixel offset in its associated row */
-		uint16 GetSrd2Offset() const;
-
-		/** Sets SRD2 associated row number */
-		void SetSrd2Offset(uint16 Offset);
+		uint8 ContributingSourceCount   : 4;
+		uint8 ExtensionBit              : 1;
+		uint8 PaddingBit                : 1;
+		uint8 Version                   : 2;
+		uint8 PayloadType               : 7;
+		uint8 MarkerBit                 : 1;
+		uint16 SequenceNumber           : 16;
+		uint32 Timestamp                : 32;
+		uint32 SynchronizationSource    : 32;
+		uint16 ExtendedSequenceNumber   : 16;
+
+		FRawSRD SRD1;                   // 48
+
+		FRawSRD SRD2;                   // 48
 
 		/** Size of RTP representation whether it has one or two SRDs */
 		static constexpr uint32 OneSRDSize = 20;
diff --git a/Engine/Plugins/VirtualProduction/Rivermax/RivermaxCore/Source/RivermaxCore/Private/RivermaxFrameAllocator.cpp b/Engine/Plugins/VirtualProduction/Rivermax/RivermaxCore/Source/RivermaxCore/Private/RivermaxFrameAllocator.cpp
--- a/Engine/Plugins/VirtualProduction/Rivermax/RivermaxCore/Source/RivermaxCore/Private/RivermaxFrameAllocator.cpp
+++ b/Engine/Plugins/VirtualProduction/Rivermax/RivermaxCore/Source/RivermaxCore/Private/RivermaxFrameAllocator.cpp
@@ -153,12 +153,6 @@ namespace UE::RivermaxCore::Private
 			AllocatedFrameSize = DesiredFrameSize;
 		}
 
-		for (int32 Index = 0; Index < AllocatedFrameCount; ++Index)
-		{
-			TSharedPtr<FRivermaxOutputFrame> Frame = MakeShared<FRivermaxOutputFrame>(Index);
-			Frame->VideoBuffer = GetFrameAddress(Index);
-		}
-
 		CudaModule.DriverAPI()->cuCtxPopCurrent(nullptr);
 
 		return true;
@@ -252,7 +246,7 @@ namespace UE::RivermaxCore::Private
 			{
 				TRACE_CPUPROFILER_EVENT_SCOPE(Rmax::CudaMemcopyAsync);
 				const CUdeviceptr CudaMemoryPointer = reinterpret_cast<CUdeviceptr>(SourceMemory);
-				Result = CudaModule.DriverAPI()->cuMemcpyDtoDAsync(reinterpret_cast<CUdeviceptr>(Args.DestinationMemory), CudaMemoryPointer, Args.SizeToCopy, reinterpret_cast<CUstream>(GPUStream));
+				Result = CudaModule.DriverAPI()->cuMemcpyDtoDAsync(DestinationCudaPointer, CudaMemoryPointer, Args.SizeToCopy, reinterpret_cast<CUstream>(GPUStream));
 				if (Result != CUDA_SUCCESS)
 				{
 					UE_LOG(LogRivermax, Error, TEXT("Failed to copy captured buffer to cuda memory. Stopping capture. Error: %d"), Result);
diff --git a/Engine/Plugins/VirtualProduction/Rivermax/RivermaxCore/Source/RivermaxCore/Private/RivermaxFrameAllocator.h b/Engine/Plugins/VirtualProduction/Rivermax/RivermaxCore/Source/RivermaxCore/Private/RivermaxFrameAllocator.h
--- a/Engine/Plugins/VirtualProduction/Rivermax/RivermaxCore/Source/RivermaxCore/Private/RivermaxFrameAllocator.h
+++ b/Engine/Plugins/VirtualProduction/Rivermax/RivermaxCore/Source/RivermaxCore/Private/RivermaxFrameAllocator.h
@@ -36,13 +36,13 @@ namespace UE::RivermaxCore::Private
 		void* DestinationMemory = nullptr;
 		
 		/** Bytes needed to be copied */
-		uint32 SizeToCopy = 0;
+		SIZE_T SizeToCopy = 0;
 		
 		/** Sidecar to be re-provided on callback after memory has been copied */
 		TSharedPtr<FBaseDataCopySideCar> SideCar;
 	};
 
-	struct FRivermaxOutputFrame;
+	class FRivermaxOutputFrame;
 
 	/** Delegate triggered when data has been copied into a given frame */
 	DECLARE_DELEGATE_OneParam(FOnFrameDataCopiedDelegate, const TSharedPtr<FBaseDataCopySideCar>&);
diff --git a/Engine/Plugins/VirtualProduction/Rivermax/RivermaxCore/Source/RivermaxCore/Private/RivermaxFrameManager.cpp b/Engine/Plugins/VirtualProduction/Rivermax/RivermaxCore/Source/RivermaxCore/Private/RivermaxFrameManager.cpp
--- a/Engine/Plugins/VirtualProduction/Rivermax/RivermaxCore/Source/RivermaxCore/Private/RivermaxFrameManager.cpp
+++ b/Engine/Plugins/VirtualProduction/Rivermax/RivermaxCore/Source/RivermaxCore/Private/RivermaxFrameManager.cpp
@@ -53,20 +53,18 @@ namespace UE::RivermaxCore::Private
 
 		if (MemoryLocation != EFrameMemoryLocation::None)
 		{
-			// Create frame state tracking containers
-			FreeFrames.Reserve(TotalFrameCount);
-			PendingFrames.Reserve(TotalFrameCount);
-			ReadyFrames.Reserve(TotalFrameCount);
+			// Create frame pool.
+			FramePool = MakeUnique<FRivermaxOutputFramePool>(TotalFrameCount);
+			TArray<TSharedPtr<FRivermaxOutputFrame>> AllocatedFrames;
 			for (uint32 Index = 0; Index < TotalFrameCount; ++Index)
 			{
-				// All frames default to being available
-				FreeFrames.Add(Index);
-
 				/** Create actual frames and assign their video memory address from allocator */
-				TSharedPtr<FRivermaxOutputFrame> Frame = MakeShared<FRivermaxOutputFrame>(Index);
+				TSharedPtr<FRivermaxOutputFrame> Frame = FramePool->AcquireShared(true /*allocate memory*/);
 				Frame->VideoBuffer = FrameAllocator->GetFrameAddress(Index);
-				Frames.Add(MoveTemp(Frame));
+				AllocatedFrames.Add(Frame);
 			}
+			AllocatedFrames.Empty();
+			FramePool->Tick();
 		}
 
 		return MemoryLocation;
@@ -79,150 +77,76 @@ namespace UE::RivermaxCore::Private
 			FrameAllocator->Deallocate();
 			FrameAllocator.Reset();
 		}
-	}
 
-	TSharedPtr<UE::RivermaxCore::Private::FRivermaxOutputFrame> FFrameManager::GetFreeFrame()
-	{
-		FScopeLock Lock(&ContainersCritSec);
-		if (!FreeFrames.IsEmpty())
+		if (FramePool.IsValid())
 		{
-			return Frames[FreeFrames[0]];
+			FramePool->Reset();
+			FramePool.Reset();
 		}
-
-		return nullptr;
 	}
 
-	void FFrameManager::MarkAsPending(const TSharedPtr<FRivermaxOutputFrame>& Frame)
+	TSharedPtr<FRivermaxOutputFrame> FFrameManager::GetFreeFrame()
 	{
-		FScopeLock Lock(&ContainersCritSec);
-
-		ensure(FreeFrames.RemoveSingle(Frame->FrameIndex) >= 1);
-		PendingFrames.Add(Frame->FrameIndex);
+		return FramePool->AcquireShared(false /*don't allocate new items*/);
 	}
 
-	TSharedPtr<UE::RivermaxCore::Private::FRivermaxOutputFrame> FFrameManager::GetNextFrame(uint32 NextFrameIdentifier)
+	TSharedPtr<FRivermaxOutputFrame> FFrameManager::DequeueFrameToSend()
 	{
 		FScopeLock Lock(&ContainersCritSec);
-		
-		TSharedPtr<FRivermaxOutputFrame> NextFrame = GetPendingFrame(NextFrameIdentifier);
-		if (NextFrame)
-		{
-			// We found a pending / reserved frame matching identifier
-			return NextFrame;
-		}
-
-		//Otherwise, prepare next free frame
-		NextFrame = GetFreeFrame();
-		if (NextFrame)
+		if (FramesToBeSent.IsEmpty())
 		{
-			NextFrame->Reset();
-			NextFrame->FrameIdentifier = NextFrameIdentifier;
-			MarkAsPending(NextFrame);
+			return nullptr;
 		}
-
-		return NextFrame;
+		TSharedPtr<FRivermaxOutputFrame> ReadyToSendFrame;
+		FramesToBeSent.Dequeue(ReadyToSendFrame);
+		return MoveTemp(ReadyToSendFrame);
 	}
 
-	TSharedPtr<FRivermaxOutputFrame> FFrameManager::GetPendingFrame(uint32 FrameIdentifier)
+	bool FFrameManager::IsFrameAvailableToSend()
 	{
 		FScopeLock Lock(&ContainersCritSec);
-		if (!PendingFrames.IsEmpty())
-		{
-			for (const uint32 FrameIndex : PendingFrames)
-			{
-				TSharedPtr<FRivermaxOutputFrame> Frame = Frames[FrameIndex];
-				if (Frame->FrameIdentifier == FrameIdentifier)
-				{
-					// We found a reserved / pending frame corresponding to next frame identifier
-					return Frame;
-				}
-			}
-		}
-		return nullptr;
+		return !FramesToBeSent.IsEmpty();
 	}
 
-	TSharedPtr<FRivermaxOutputFrame> FFrameManager::GetReadyFrame()
+	void FFrameManager::FrameSentEvent()
 	{
-		FScopeLock Lock(&ContainersCritSec);
-		if (!ReadyFrames.IsEmpty())
-		{
-			const uint32 NextReadyFrame = ReadyFrames[0];
-			return Frames[NextReadyFrame];
-		}
-	
-		return nullptr;
-	}
-
-	void FFrameManager::MarkAsSent(const TSharedPtr<FRivermaxOutputFrame>& Frame)
-	{
-		{
-			FScopeLock Lock(&ContainersCritSec);
-
-			Frame->Reset();
-			FreeFrames.Add(SendingFrame);
-
-			if (ensure(Frame->FrameIndex == SendingFrame))
-			{
-				SendingFrame = INDEX_NONE;
-			}
-		}
-		
+		// Todo: the frame iself should call this when it is returned back to the pool.
 		OnFreeFrameDelegate.ExecuteIfBound();
 	}
 
-	void FFrameManager::MarkAsReady(const TSharedPtr<FRivermaxOutputFrame>& Frame)
+	void FFrameManager::EnqueFrameToSend(const TSharedPtr<FRivermaxOutputFrame>& Frame)
 	{
 		{
 			// Make frame available to be sent
 			FScopeLock Lock(&ContainersCritSec);
-			Frame->ReadyTimestamp = RivermaxManager->GetTime();
-			PendingFrames.RemoveSingle(Frame->FrameIndex);
-			ReadyFrames.Add(Frame->FrameIndex);
+			FramesToBeSent.Enqueue(Frame);
 		}
 
 		OnFrameReadyDelegate.ExecuteIfBound();
 	}
 
-	void FFrameManager::MarkAsSending(const TSharedPtr<FRivermaxOutputFrame>& Frame)
-	{
-		// Make frame available to be sent
-		FScopeLock Lock(&ContainersCritSec);
-		if (ensure(!ReadyFrames.IsEmpty()))
-		{
-			SendingFrame = ReadyFrames[0];
-			ensure(SendingFrame == Frame->FrameIndex);
-			ReadyFrames.RemoveAt(0);
-		}
-	}
-
-	const TSharedPtr<FRivermaxOutputFrame> FFrameManager::GetFrame(int32 Index) const
-	{
-		return Frames[Index];
-	}
-
-	bool FFrameManager::SetFrameData(const FRivermaxOutputVideoFrameInfo& NewFrameInfo)
+	bool FFrameManager::SetFrameData(const FRivermaxOutputVideoFrameInfo& NewFrameInfo, TSharedPtr<FRivermaxOutputFrame> ReservedFrame)
 	{
 		bool bSuccess = false;
-		if (TSharedPtr<FRivermaxOutputFrame> NextFrame = GetNextFrame(NewFrameInfo.FrameIdentifier))
-		{
-			TSharedPtr<FFrameBufferCopyInfo> Sidecar = MakeShared<FFrameBufferCopyInfo>();
-			Sidecar->CopiedFrame = NextFrame;
+		check(ReservedFrame.IsValid());
 
-			FCopyArgs Args;
-			Args.RHISourceMemory = NewFrameInfo.GPUBuffer;
-			Args.SourceMemory = NewFrameInfo.VideoBuffer;
-			Args.DestinationMemory = NextFrame->VideoBuffer;
-			Args.SizeToCopy = NewFrameInfo.Height * NewFrameInfo.Stride;
-			Args.SideCar = MoveTemp(Sidecar);
+		TSharedPtr<FFrameBufferCopyInfo> Sidecar = MakeShared<FFrameBufferCopyInfo>();
+		Sidecar->CopiedFrame = ReservedFrame;
 
-			bSuccess = FrameAllocator->CopyData(Args);
+		FCopyArgs Args;
+		Args.RHISourceMemory = NewFrameInfo.GPUBuffer;
+		Args.SourceMemory = NewFrameInfo.CPUBuffer;
+		Args.DestinationMemory = ReservedFrame->VideoBuffer;
+		Args.SizeToCopy = NewFrameInfo.Height * NewFrameInfo.Stride;
+		Args.SideCar = MoveTemp(Sidecar);
+		bSuccess = FrameAllocator->CopyData(Args);
+		ReservedFrame->SetFrameCounter(NewFrameInfo.FrameIdentifier);
 
-			if (!bSuccess)
-			{
-				OnCriticalErrorDelegate.ExecuteIfBound();
-			}
-			 
+		if (!bSuccess)
+		{
+			OnCriticalErrorDelegate.ExecuteIfBound();
 		}
+			 
 		return bSuccess;
 	}
 
@@ -231,22 +155,12 @@ namespace UE::RivermaxCore::Private
 		TSharedPtr<FFrameBufferCopyInfo> CopyInfo = StaticCastSharedPtr<FFrameBufferCopyInfo>(Payload);
 		if (ensure(CopyInfo && CopyInfo->CopiedFrame))
 		{
+			//TODO: make sure that tracing is not reliant on frame index.
+			TRACE_CPUPROFILER_EVENT_SCOPE_TEXT(*FRivermaxTracingUtils::RmaxOutFrameReadyTraceEvents[CopyInfo->CopiedFrame->GetFrameCounter() % 10]);
+			TRACE_CPUPROFILER_EVENT_SCOPE_TEXT(*FRivermaxTracingUtils::RmaxOutMediaCapturePipeTraceEvents[CopyInfo->CopiedFrame->GetFrameCounter() % 10]);
 			OnPreFrameReadyDelegate.ExecuteIfBound();
-
-			if (TSharedPtr<FRivermaxOutputFrame> AvailableFrame = GetNextFrame(CopyInfo->CopiedFrame->FrameIdentifier))
-			{
-				// Video frame has been copied, update frame's state
-				AvailableFrame->bIsVideoBufferReady = true;
-				if (AvailableFrame->IsReadyToBeSent())
-				{
-					TRACE_CPUPROFILER_EVENT_SCOPE_TEXT(*FRivermaxTracingUtils::RmaxOutFrameReadyTraceEvents[AvailableFrame->FrameIndex % 10]);
-					TRACE_CPUPROFILER_EVENT_SCOPE_TEXT(*FRivermaxTracingUtils::RmaxOutMediaCapturePipeTraceEvents[AvailableFrame->FrameIdentifier % 10]);
-
-					MarkAsReady(AvailableFrame);
-				}
-			}
+			EnqueFrameToSend(CopyInfo->CopiedFrame);
 		}
-		
 	}
 }
 
diff --git a/Engine/Plugins/VirtualProduction/Rivermax/RivermaxCore/Source/RivermaxCore/Private/RivermaxFrameManager.h b/Engine/Plugins/VirtualProduction/Rivermax/RivermaxCore/Source/RivermaxCore/Private/RivermaxFrameManager.h
--- a/Engine/Plugins/VirtualProduction/Rivermax/RivermaxCore/Source/RivermaxCore/Private/RivermaxFrameManager.h
+++ b/Engine/Plugins/VirtualProduction/Rivermax/RivermaxCore/Source/RivermaxCore/Private/RivermaxFrameManager.h
@@ -9,6 +9,7 @@
 #include "RivermaxFrameAllocator.h"
 #include "RivermaxOutputFrame.h"
 #include "RivermaxTypes.h"
+#include "MediaObjectPool.h"
 
 namespace UE::RivermaxCore
 {
@@ -72,6 +73,13 @@ namespace UE::RivermaxCore::Private
 		FOnCriticalErrorDelegate OnCriticalErrorDelegate;
 	};
 
+	/**
+	* Media object pool responsible for managing shared pointers. Pointers allocated by this class 
+	* are returned to the pending pool where it is checked via IsReadyToBeUsed if it is ready to be used again.
+	* It reduces manual thread/frame management required.
+	*/
+	class FRivermaxOutputFramePool : public TMediaObjectPool<FRivermaxOutputFrame> { };
+
 	/** 
 	 * Class managing frames that we output over network 
 	 * Handles memory allocation and state tracking
@@ -107,36 +115,20 @@ namespace UE::RivermaxCore::Private
 		/** Returns a frame not being used */
 		TSharedPtr<FRivermaxOutputFrame> GetFreeFrame();
 
-		/**
-		 * Get next frame that can be used. 
-		 * If a pending one matching identifier is found, it is returned.
-		 * If not, a free frame will be returned, if any are available
-		 */
-		TSharedPtr<FRivermaxOutputFrame> GetNextFrame(uint32 NextFrameIdentifier);
-
-		/** Returns a pending frame with the given identifier. If none is found, returns nullptr */
-		TSharedPtr<FRivermaxOutputFrame> GetPendingFrame(uint32 FrameIdentifier);
-		
 		/** Returns next frame ready to be sent. Can be null. */
-		TSharedPtr<FRivermaxOutputFrame> GetReadyFrame();
-
-		/** Returns a frame for a given index */
-		const TSharedPtr<FRivermaxOutputFrame> GetFrame(int32 Index) const;
-
-		/** Mark a frame as being used and data should be coming soon */
-		void MarkAsPending(const TSharedPtr<FRivermaxOutputFrame>& Frame);
+		TSharedPtr<FRivermaxOutputFrame> DequeueFrameToSend();
 		
 		/** Mark a frame as being ready to be sent */
-		void MarkAsReady(const TSharedPtr<FRivermaxOutputFrame>& Frame);
-		
-		/** Mark a frame as being sent */
-		void MarkAsSending(const TSharedPtr<FRivermaxOutputFrame>& Frame);
+		void EnqueFrameToSend(const TSharedPtr<FRivermaxOutputFrame>& Frame);
+
+		/** Returns next frame ready to be sent. Can be null. */
+		bool IsFrameAvailableToSend();
 
 		/** Mark a frame as sent */
-		void MarkAsSent(const TSharedPtr<FRivermaxOutputFrame>& Frame);
+		void FrameSentEvent();
 
 		/** Initiates memory copy for a given frame */
-		bool SetFrameData(const FRivermaxOutputVideoFrameInfo& NewFrame);
+		bool SetFrameData(const FRivermaxOutputVideoFrameInfo& NewFrame, TSharedPtr<FRivermaxOutputFrame> ReservedFrame);
 
 	private:
 
@@ -160,8 +152,13 @@ namespace UE::RivermaxCore::Private
 		/** Frame allocator dealing with memory operation */
 		TUniquePtr<FBaseFrameAllocator> FrameAllocator;
 
-		/** List of allocated frames. */
-		TArray<TSharedPtr<FRivermaxOutputFrame>> Frames;
+		/** 
+		* Pool of allocated frames. Thread safe allocation. All items allocated via AcquireShared are returned back to the pool automatically.
+		*/
+		TUniquePtr<FRivermaxOutputFramePool> FramePool;
+
+		/** Queue of frames that are ready to be sent. First in, first out. */
+		TQueue<TSharedPtr<FRivermaxOutputFrame>> FramesToBeSent;
 
 		/** Delegate triggered when a frame is free to use */
 		FOnFreeFrameDelegate OnFreeFrameDelegate;
@@ -174,18 +171,6 @@ namespace UE::RivermaxCore::Private
 
 		/** Delegate triggered when a critical error has happened and stream should shut down */
 		FOnCriticalErrorDelegate OnCriticalErrorDelegate;
-		
-		/** List of available frames. */
-		TArray<int32> FreeFrames;
-		
-		/** List of frames being used but not ready to be sent yet  */
-		TArray<int32> PendingFrames;
-
-		/** List of frames ready to send */
-		TArray<int32> ReadyFrames;
-
-		/** Frame currently being sent */
-		int32 SendingFrame = 0;
 
 		/** Quick access to rivermax manager */
 		TSharedPtr<IRivermaxManager> RivermaxManager;
diff --git a/Engine/Plugins/VirtualProduction/Rivermax/RivermaxCore/Source/RivermaxCore/Private/RivermaxInputStream.cpp b/Engine/Plugins/VirtualProduction/Rivermax/RivermaxCore/Source/RivermaxCore/Private/RivermaxInputStream.cpp
--- a/Engine/Plugins/VirtualProduction/Rivermax/RivermaxCore/Source/RivermaxCore/Private/RivermaxInputStream.cpp
+++ b/Engine/Plugins/VirtualProduction/Rivermax/RivermaxCore/Source/RivermaxCore/Private/RivermaxInputStream.cpp
@@ -328,9 +328,6 @@ namespace UE::RivermaxCore::Private
 			RawHeaderPtr += StrideIndex * BufferConfiguration.HeaderStride;
 			DataPtr += StrideIndex * BufferConfiguration.DataStride;
 
-			// When using RMX_INPUT_RAW_PACKET, app head is preceded by net header
-			const uint8_t* AppHeader = RawHeaderPtr; // App header. If the stream type is RMX_INPUT_RAW_PACKET then app header is preceded by net header
-
 			const rmx_input_packet_info* PacketInfo = CachedAPI->rmx_input_get_packet_info(&BufferConfiguration.ChunkHandle, StrideIndex);
 			const size_t PacketSize = rmx_input_get_packet_size(PacketInfo, BufferConfiguration.DataBlockID);
 			const size_t HeaderSize = rmx_input_get_packet_size(PacketInfo, BufferConfiguration.HeaderBlockID);
diff --git a/Engine/Plugins/VirtualProduction/Rivermax/RivermaxCore/Source/RivermaxCore/Private/RivermaxOutputFrame.cpp b/Engine/Plugins/VirtualProduction/Rivermax/RivermaxCore/Source/RivermaxCore/Private/RivermaxOutputFrame.cpp
--- a/Engine/Plugins/VirtualProduction/Rivermax/RivermaxCore/Source/RivermaxCore/Private/RivermaxOutputFrame.cpp
+++ b/Engine/Plugins/VirtualProduction/Rivermax/RivermaxCore/Source/RivermaxCore/Private/RivermaxOutputFrame.cpp
@@ -4,32 +4,21 @@
 
 namespace UE::RivermaxCore::Private
 {
-	FRivermaxOutputFrame::FRivermaxOutputFrame(uint32 InFrameIndex)
-		: FrameIndex(InFrameIndex)
-		, VideoBuffer(nullptr)
+	FRivermaxOutputFrame::FRivermaxOutputFrame()
+		: VideoBuffer(nullptr)
 	{
 		Clear();
 	}
 
-	bool FRivermaxOutputFrame::IsReadyToBeSent() const
-	{
-		return FrameIdentifier != InvalidIdentifier && bIsVideoBufferReady;
-	}
-
 	void FRivermaxOutputFrame::Clear()
 	{
 		PacketCounter = 0;
 		LineNumber = 0;
-		SRDOffset = 0;
-		BytesSent = 0;
 	}
 
 	void FRivermaxOutputFrame::Reset()
 	{
-		FrameIdentifier = InvalidIdentifier;
-		bIsVideoBufferReady = false;
 		bCaughtTimingIssue = false;
-		ReadyTimestamp = 0;
 		Clear();
 	}
 
diff --git a/Engine/Plugins/VirtualProduction/Rivermax/RivermaxCore/Source/RivermaxCore/Private/RivermaxOutputFrame.h b/Engine/Plugins/VirtualProduction/Rivermax/RivermaxCore/Source/RivermaxCore/Private/RivermaxOutputFrame.h
--- a/Engine/Plugins/VirtualProduction/Rivermax/RivermaxCore/Source/RivermaxCore/Private/RivermaxOutputFrame.h
+++ b/Engine/Plugins/VirtualProduction/Rivermax/RivermaxCore/Source/RivermaxCore/Private/RivermaxOutputFrame.h
@@ -6,52 +6,40 @@
 
 #include "Math/NumericLimits.h"
 #include "Misc/Timecode.h"
+#include "MediaObjectPool.h"
 
 namespace UE::RivermaxCore::Private
 {
 	/**
 	 * Output frame descriptor. Contains data to be sent on the wire and packetization tracking
 	 */
-	struct FRivermaxOutputFrame
+	class FRivermaxOutputFrame : public IMediaPoolable
 	{
-		FRivermaxOutputFrame(uint32 InFrameIndex);
+	public:
+		FRivermaxOutputFrame();
 		FRivermaxOutputFrame(const FRivermaxOutputFrame&) = delete;
 		FRivermaxOutputFrame& operator=(const FRivermaxOutputFrame&) = delete;
 
-		/** Whether this frame is ready to be sent */
-		bool IsReadyToBeSent() const;
-
 		/** Reset internal counters to make it resendable */
 		void Clear();
 
 		/** Reset identifier and internal counters to make it available */
 		void Reset();
 
-		/** Index of the frame used as unique identifier */
-		uint32 FrameIndex = 0;
-
 		/** 
 		 * Video buffer where we copy texture to be sent out 
 		 * If using GPUDirect, memory will be allocated in cuda space
 		 */
 		void* VideoBuffer = nullptr;
 
-		/** Identifier of the frame incremented every time a frame is captured */
-		static constexpr uint32 InvalidIdentifier = MAX_uint32;
-		uint32 FrameIdentifier = InvalidIdentifier;
-
 		/** Timecode at which this frame was captured */
 		FTimecode Timecode;
 
-		/** Used to flag this frame ready to be sent */
-		bool bIsVideoBufferReady = false;
-
 		/** Variables about progress of packetization of the frame */
 		uint32 PacketCounter = 0;
 		uint32 LineNumber = 0;
 		uint16 SRDOffset = 0;
 		uint32 ChunkNumber = 0;
-		uint32 BytesSent = 0;
 
 		/** Timestamp of this frame used for RTP headers  */
 		uint32 MediaTimestamp = 0;
@@ -66,16 +54,44 @@ namespace UE::RivermaxCore::Private
 		void* FrameStartPtr = nullptr;
 		
 		/** Offset in the frame where we are at to copy next block of data */
-		uint32 Offset = 0;
-		
-		/** Time covered by memcopies. Not used for now but it could be used to detect if memory will be used before it's actually copied. */
-		uint64 ScheduleTimeCopied = 0;
-
-		/** Time at which this frame was made available to be sent */
-		uint64 ReadyTimestamp = 0;
+		SIZE_T Offset = 0;
 
 		/** Whether timing issues were detected while sending frame out. If yes, we skip the next frame boundary */
 		bool bCaughtTimingIssue = false;
+
+	public:
+		/**
+		* Set frame counter .
+		*/
+		void SetFrameCounter(uint64 InFrameCounter)
+		{
+			FrameCounter = InFrameCounter;
+		}
+
+		/**
+		* Returns frame counter corresponding to the data captured by Media Capture.
+		*/
+		uint64 GetFrameCounter()
+		{
+			return FrameCounter;
+		}
+
+		/**
+		 * Used to check if returned object is ready for reuse right away
+		 */
+		virtual bool IsReadyForReuse()
+		{
+			return true;
+		}
+
+		/**
+		 * Used to check if returned object is ready for reuse right away
+		 */
+		virtual void ShutdownPoolable() {}
+
+	private:
+		/** Number corresponding to the engine frame counter of the frame captured by Media Capture.*/
+		uint64 FrameCounter;
 	};
 }
 
diff --git a/Engine/Plugins/VirtualProduction/Rivermax/RivermaxCore/Source/RivermaxCore/Private/RivermaxOutputStream.cpp b/Engine/Plugins/VirtualProduction/Rivermax/RivermaxCore/Source/RivermaxCore/Private/RivermaxOutputStream.cpp
--- a/Engine/Plugins/VirtualProduction/Rivermax/RivermaxCore/Source/RivermaxCore/Private/RivermaxOutputStream.cpp
+++ b/Engine/Plugins/VirtualProduction/Rivermax/RivermaxCore/Source/RivermaxCore/Private/RivermaxOutputStream.cpp
@@ -16,7 +16,7 @@
 #include "RivermaxTypes.h"
 #include "RivermaxUtils.h"
 
-
+#define GET_FRAME_INDEX(OUTPUT_FRAME) OUTPUT_FRAME->GetFrameCounter() % Options.NumberOfBuffers
 namespace UE::RivermaxCore::Private
 {
 	static TAutoConsoleVariable<int32> CVarRivermaxWakeupOffset(
@@ -112,17 +112,6 @@ namespace UE::RivermaxCore::Private
 		TEXT("Optimization used to prefill every RTP headers with known data."),
 		ECVF_Default);
 
-	static TAutoConsoleVariable<bool> CVarRivermaxOutputForceSkip(
-		TEXT("Rivermax.Output.ForceSkip"), false,
-		TEXT("Used to prevent enforced blank window when multiplier is used to prevent timing issues.")
-		TEXT("Only affects alignment point method."),
-		ECVF_Default);
-
-	static TAutoConsoleVariable<float> CVarRivermaxOutputFrameRateMultiplier(
-		TEXT("Rivermax.Output.FrameRateMultiplier"), 1.0,
-		TEXT("Multiplier applied to desired output frame rate in order to reduce time it takes to send out a frame and slowly correct misalignment that could happen."),
-		ECVF_Default);
-
 	static TAutoConsoleVariable<float> CVarRivermaxOutputCommitChunksOffsetPercent(
 		TEXT("Rivermax.Output.CommitChunksOffsetPercent"), 0.3,
 		TEXT("This CVar will allow for chunks to be committed before the next alignment point if it is at all possible.\n\
@@ -179,6 +168,83 @@ namespace UE::RivermaxCore::Private
 
 		return false;
 	}
+	
+	void UpdateVideoRtpHeader
+		( FRawRTPHeader& OutHeader
+		, uint32& OutPacketCounter
+		, uint16& OutSRDOffset
+		, uint32& OutLineNumber
+		, uint32 InSequenceNumber
+		, uint32 InTimeStamp
+		, const FIntPoint& AlignedResolution
+		, const FVideoFormatInfo& InFormatInfo
+		, const FRivermaxOutputStreamMemory& InStreamMemory
+		)
+	{
+		OutHeader = {};
+		OutHeader.Version = 2;
+		OutHeader.PaddingBit = 0;
+		OutHeader.ExtensionBit = 0;
+		OutHeader.PayloadType = 96; //Payload type should probably be inferred from SDP
+
+		OutHeader.SequenceNumber = ByteSwap((uint16)(InSequenceNumber & 0xFFFF));
+		OutHeader.Timestamp = ByteSwap(InTimeStamp);
+		OutHeader.ExtendedSequenceNumber = ByteSwap((uint16)((InSequenceNumber >> 16) & 0xFFFF));
+
+		//2110 specific header
+		OutHeader.SynchronizationSource = ByteSwap((uint32)0x0eb51dbd);  // Should Unreal has its own synch source ID
+
+		{
+			if (OutPacketCounter + 1 == InStreamMemory.PacketsPerFrame)
+			{
+				OutHeader.MarkerBit = 1; // last packet in frame (Marker)
+			}
+
+			// Verify if payload size exceeds line 
+			const uint32 CurrentPayloadSize = InStreamMemory.PayloadSizes[OutPacketCounter];
+
+			const uint32 LineSizeOffset = ((OutSRDOffset / InFormatInfo.PixelGroupCoverage) * InFormatInfo.PixelGroupSize);
+			const uint32 LineSize = ((AlignedResolution.X / InFormatInfo.PixelGroupCoverage) * InFormatInfo.PixelGroupSize);
+
+			const uint16 SRD1Length = FMath::Min(LineSize - LineSizeOffset, CurrentPayloadSize);
+			const uint16 SRD1PixelCount = SRD1Length / InFormatInfo.PixelGroupSize * InFormatInfo.PixelGroupCoverage;
+			uint16 SRD2Length = SRD1Length < CurrentPayloadSize ? CurrentPayloadSize - SRD1Length : 0;
+			if (SRD2Length && OutLineNumber == ((uint32)AlignedResolution.Y - 1))
+			{
+				SRD2Length = 0;
+			}
+
+			OutHeader.SRD1.Length = ByteSwap(SRD1Length);
+			OutHeader.SRD1.SetRowNumber(OutLineNumber); //todo divide by 2 if interlaced
+			OutHeader.SRD1.SetOffset(OutSRDOffset);
+			OutHeader.SRD1.ContinuationBit = SRD2Length > 0;
+			OutHeader.SRD1.FieldIdentification = 0; //todo when fields are sent for interlace
+
+			OutSRDOffset += SRD1PixelCount;
+			if (OutSRDOffset >= AlignedResolution.X)
+			{
+				OutSRDOffset = 0;
+				++OutLineNumber;
+			}
+
+			if (SRD2Length > 0)
+			{
+				OutHeader.SRD2.Length = ByteSwap(SRD2Length);
+				OutHeader.SRD2.SetRowNumber(OutLineNumber);
+				OutHeader.SRD2.SetOffset(OutSRDOffset);
+				OutHeader.SRD2.ContinuationBit = 0;
+				OutHeader.SRD2.FieldIdentification = 0;
+
+				const uint16 SRD2PixelCount = SRD2Length / InFormatInfo.PixelGroupSize * InFormatInfo.PixelGroupCoverage;
+				OutSRDOffset += SRD2PixelCount;
+				if (OutSRDOffset >= AlignedResolution.X)
+				{
+					OutSRDOffset = 0;
+					++OutLineNumber;
+				}
+			}
+		}
+	}
 
 	uint32 FindLinesPerChunk(const FRivermaxOutputStreamOptions& InOptions)
 	{
@@ -373,13 +439,6 @@ namespace UE::RivermaxCore::Private
 
 	}
 
-	/** Sidecar used when copying  parts of frame data while scheduling */
-	struct FMemChunkCopiedInfo : public FBaseDataCopySideCar
-	{
-		/** Time covered inside frame interval including last chunk copied. */
-		uint64 TimeCovered = 0;
-	};
-
 	struct FRTPHeaderPrefiller
 	{
 		FRTPHeaderPrefiller(FRivermaxOutputStream& InRmaxOutputStream)
@@ -393,8 +452,6 @@ namespace UE::RivermaxCore::Private
 		{
 			using namespace UE::RivermaxCore::Private::Utils;
 
-			const uint32 Width = Stream.Options.AlignedResolution.X;
-			const uint32 Height = Stream.Options.AlignedResolution.Y;
 			const int32 PacketCount = Stream.StreamMemory.ChunksPerFrameField * Stream.StreamMemory.PacketsPerChunk;
 			uint32 BufferIndex = 0;
 			for (int32 MemblockIndex = 0; MemblockIndex < Stream.StreamMemory.RTPHeaders.Num(); ++MemblockIndex)
@@ -407,65 +464,18 @@ namespace UE::RivermaxCore::Private
 					if (ensure(Stream.StreamMemory.RTPHeaders[MemblockIndex].IsValidIndex(HeaderIndex)))
 					{
 						FRawRTPHeader& CurrentHeader = Stream.StreamMemory.RTPHeaders[MemblockIndex][HeaderIndex];
-						CurrentHeader = {};
-						CurrentHeader.Version = 2;
-						CurrentHeader.PaddingBit = 0;
-						CurrentHeader.ExtensionBit = 0;
-						CurrentHeader.PayloadType = 96; //Payload type should probably be inferred from SDP
-
-						//2110 specific header
-						CurrentHeader.SynchronizationSource = ByteSwap((uint32)0x0eb51dbd);  // Should Unreal has its own synch source ID
-
-						{
-							if (PacketIndex + 1 == Stream.StreamMemory.PacketsPerFrame)
-							{
-								CurrentHeader.MarkerBit = 1; // last packet in frame (Marker)
-							}
-
-							// Verify if payload size exceeds line 
-							const uint32 CurrentPayloadSize = Stream.StreamMemory.PayloadSizes[PacketIndex];
-
-							const uint32 LineSizeOffset = ((SRDOffset / Stream.FormatInfo.PixelGroupCoverage) * Stream.FormatInfo.PixelGroupSize);
-							const uint32 LineSize = ((Width / Stream.FormatInfo.PixelGroupCoverage) * Stream.FormatInfo.PixelGroupSize);
-
-							const uint16 SRD1Length = FMath::Min(LineSize - LineSizeOffset, CurrentPayloadSize);
-							const uint16 SRD1PixelCount = SRD1Length / Stream.FormatInfo.PixelGroupSize * Stream.FormatInfo.PixelGroupCoverage;
-							uint16 SRD2Length = SRD1Length < CurrentPayloadSize ? CurrentPayloadSize - SRD1Length : 0;
-							if (SRD2Length && LineNumber == (Height - 1))
-							{
-								SRD2Length = 0;
-							}
-
-							CurrentHeader.SRD1Length = ByteSwap(SRD1Length);
-							CurrentHeader.SetSrd1RowNumber(LineNumber); //todo divide by 2 if interlaced
-							CurrentHeader.FieldIdentification1 = 0; //todo when fields are sent for interlace
-							CurrentHeader.SetSrd1Offset(SRDOffset);
-
-							SRDOffset += SRD1PixelCount;
-							if (SRDOffset >= Width)
-							{
-								SRDOffset = 0;
-								++LineNumber;
-							}
-
-							if (SRD2Length > 0)
-							{
-								CurrentHeader.SRD2Length = ByteSwap(SRD2Length);
-
-								CurrentHeader.ContinuationBit1 = 1;
-								CurrentHeader.FieldIdentification2 = 0;
-								CurrentHeader.SetSrd2RowNumber(LineNumber);
-								CurrentHeader.SetSrd2Offset(SRDOffset);
-
-								const uint16 SRD2PixelCount = SRD2Length / Stream.FormatInfo.PixelGroupSize * Stream.FormatInfo.PixelGroupCoverage;
-								SRDOffset += SRD2PixelCount;
-								if (SRDOffset >= Width)
-								{
-									SRDOffset = 0;
-									++LineNumber;
-								}
-							}
-						}
+						
+						UpdateVideoRtpHeader(
+							CurrentHeader,
+							PacketIndex,
+							SRDOffset,
+							LineNumber,
+							0, //SequenceNumber
+							0, //Timestamp
+							Stream.Options.AlignedResolution,
+							Stream.FormatInfo,
+							Stream.StreamMemory
+						);
 					}
 
 					++BufferIndex;
@@ -504,7 +514,7 @@ namespace UE::RivermaxCore::Private
 			CvarDelay->OnChangedDelegate().AddRaw(this, &FRivermaxOutputStream::OnCVarRandomDelayChanged);
 		}
 
-		const uint32 FrameSize = GetStride() * Options.AlignedResolution.Y;
+		const SIZE_T FrameSize = GetRowSizeInBytes() * Options.AlignedResolution.Y;
 		if (FrameSize == 0)
 		{
 			UE_LOG(LogRivermax, Warning, TEXT("Can't create Rivermax Output Stream. Frame size of 0 is invalid. Verify resolution."));
@@ -529,15 +539,13 @@ namespace UE::RivermaxCore::Private
 			return false;
 		}
 
-		InitializeTimingProtections();
-
 		// Enable frame boundary monitoring
 		MonitoringGuid = RivermaxModule->GetRivermaxBoundaryMonitor().StartMonitoring(Options.FrameRate);
 
 		InitializationFuture = Async(EAsyncExecution::TaskGraph, [this]()
 		{
 			TAnsiStringBuilder<2048> SDPDescription;
-			UE::RivermaxCore::Private::Utils::StreamOptionsToSDPDescription(Options, StreamMemory.FrameRateMultiplier, SDPDescription);
+			UE::RivermaxCore::Private::Utils::StreamOptionsToSDPDescription(Options, SDPDescription);
 
 			// Create Rivermax stream using memory configuration
 			{
@@ -669,13 +677,25 @@ namespace UE::RivermaxCore::Private
 		TRACE_CPUPROFILER_EVENT_SCOPE(FRivermaxOutputStream::PushVideoFrame);
 
 		FRivermaxOutputVideoFrameInfo CopyInfo;
-		CopyInfo.Stride = GetStride();
+		CopyInfo.Stride = GetRowSizeInBytes();
 		CopyInfo.Height = Options.AlignedResolution.Y;
 		CopyInfo.FrameIdentifier = NewFrame.FrameIdentifier;
 		CopyInfo.GPUBuffer = NewFrame.GPUBuffer;
-		CopyInfo.VideoBuffer = NewFrame.VideoBuffer;
+		CopyInfo.CPUBuffer = NewFrame.CPUBuffer;
 		CopyInfo.Width = NewFrame.Width;
-		return FrameManager->SetFrameData(CopyInfo);
+
+		// Clear reserved frame if there is one. If not, try to get free frame.
+		// ReservedFrame should always be valid when block on reservation mode is used.
+		TSharedPtr<FRivermaxOutputFrame> ReservedFrame;
+		check(ReservedFrames.Dequeue(ReservedFrame) || Options.FrameLockingMode != EFrameLockingMode::BlockOnReservation);
+		TSharedPtr<FRivermaxOutputFrame> FrameToSend = ReservedFrame.IsValid() ? MoveTemp(ReservedFrame) : FrameManager->GetFreeFrame();
+
+		// If this is invalid it means that frame locking mode is BlockOnReservation and the render ran faster than media output fps.
+		if (FrameToSend.IsValid())
+		{
+			return FrameManager->SetFrameData(CopyInfo, FrameToSend);
+		}
+		return false;
 	}
 
 	void FRivermaxOutputStream::Process_AnyThread()
@@ -730,10 +750,6 @@ namespace UE::RivermaxCore::Private
 				if (CurrentFrame->bCaughtTimingIssue)
 				{
 					++Stats.TimingIssueCount;
-				}
-
-				if (CurrentFrame->bCaughtTimingIssue || StreamMemory.bAlwaysSkipChunk)
-				{
 					constexpr uint64 ChunksToSkip = 0;
 					SkipChunks(ChunksToSkip);
 				}
@@ -748,7 +764,7 @@ namespace UE::RivermaxCore::Private
 	{
 		checkSlow(Options.AlignmentMode == ERivermaxAlignmentMode::AlignmentPoint);
 
-		TSharedPtr<FRivermaxOutputFrame> NextFrameToSend = FrameManager->GetReadyFrame();
+		TSharedPtr<FRivermaxOutputFrame> NextFrameToSend = FrameManager->DequeueFrameToSend();
 		if (ensure(NextFrameToSend))
 		{
 			InitializeNextFrame(NextFrameToSend);
@@ -769,8 +785,7 @@ namespace UE::RivermaxCore::Private
 					bool bHasDataToCopy = true;
 					while (FPlatformTime::Seconds() < TargetPlatformTimeSec && bHasDataToCopy && bIsActive)
 					{
-						bHasDataToCopy = CopyFrameData(NextFrameToSend, reinterpret_cast<uint8*>(StreamMemory.BufferAddresses[StreamData.ExpectedFrameIndex]));
-						
+						bHasDataToCopy = CopyFrameData(NextFrameToSend, reinterpret_cast<uint8*>(StreamMemory.BufferAddresses[GET_FRAME_INDEX(NextFrameToSend)]));
 					}
 				}
 
@@ -809,7 +824,6 @@ namespace UE::RivermaxCore::Private
 
 				// Make the next frame to send the current one and update its state
 				CurrentFrame = MoveTemp(NextFrameToSend);
-				FrameManager->MarkAsSending(CurrentFrame);
 			}
 		}
 		else
@@ -827,10 +841,10 @@ namespace UE::RivermaxCore::Private
 		//2110 stream type
 		//We need to use the fullframe allocated size to compute the payload size.
 
-		const int32 BytesPerLine = GetStride();
+		const SIZE_T BytesPerRow = GetRowSizeInBytes();
 
 		// By default we want to divide the bytes evenly across packets. Some resolutions will require packets to be sized unevenly.
-		const bool bFoundPayload = FindPayloadSize(Options, BytesPerLine, FormatInfo, StreamMemory.PayloadSize);
+		const bool bFoundPayload = FindPayloadSize(Options, BytesPerRow, FormatInfo, StreamMemory.PayloadSize);
 		if (bFoundPayload == false)
 		{
 			// Find out payload we want to use. Either we go the 'potential' multi SRD route or we restrict the stream based on supported resolutions.
@@ -856,7 +870,7 @@ namespace UE::RivermaxCore::Private
 
 		// With payload size in hand, figure out how many packets we will need, how many chunks (group of packets) and configure descriptor arrays
 		const uint32 PixelCount = Options.AlignedResolution.X * Options.AlignedResolution.Y;
-		const uint64 FrameSize = PixelCount / FormatInfo.PixelGroupCoverage * FormatInfo.PixelGroupSize;
+		const uint64 FrameSizeInBytes = PixelCount / FormatInfo.PixelGroupCoverage * FormatInfo.PixelGroupSize;
 
 		StreamMemory.PixelGroupPerPacket = StreamMemory.PayloadSize / FormatInfo.PixelGroupSize;
 		StreamMemory.PixelsPerPacket = StreamMemory.PixelGroupPerPacket * FormatInfo.PixelGroupCoverage;
@@ -884,6 +898,8 @@ namespace UE::RivermaxCore::Private
 		StreamMemory.RTPHeaders.SetNumZeroed(StreamMemory.MemoryBlockCount);
 		StreamMemory.PayloadSizes.SetNumZeroed(StreamMemory.PacketsPerMemoryBlock);
 		StreamMemory.HeaderSizes.SetNumZeroed(StreamMemory.PacketsPerMemoryBlock);
+
+		// Packed size should be 26 octents.
 		StreamMemory.HeaderStrideSize = sizeof(FRawRTPHeader);
 
 		if (!SetupFrameManagement())
@@ -893,6 +909,10 @@ namespace UE::RivermaxCore::Private
 
 		FRTPHeaderPrefiller RTPFiller(*this);
 
+		// Used to keep track of acquired frames to prevent the same frame fomr being acquired from the pool more than once. 
+		// When this scope is exited all frames are returned back to the Frame Manager pool.
+		TArray<TSharedPtr<UE::RivermaxCore::Private::FRivermaxOutputFrame>> ProcessedFrames;
+
 		StreamMemory.MemoryBlocks.SetNumZeroed(StreamMemory.MemoryBlockCount);
 		CachedAPI->rmx_output_media_init_mem_blocks(StreamMemory.MemoryBlocks.GetData(), StreamMemory.MemoryBlockCount);
 		for (uint32 BlockIndex = 0; BlockIndex < StreamMemory.MemoryBlockCount; ++BlockIndex)
@@ -925,7 +945,9 @@ namespace UE::RivermaxCore::Private
 			}
 			else
 			{
-				DataMemory->addr = FrameManager->GetFrame(BlockIndex)->VideoBuffer;
+				TSharedPtr<UE::RivermaxCore::Private::FRivermaxOutputFrame> ProcessedFrame = FrameManager->GetFreeFrame();
+				DataMemory->addr = ProcessedFrame->VideoBuffer;
+				ProcessedFrames.Add(ProcessedFrame);
 			}
 
 			DataMemory->length = StreamMemory.ChunksPerFrameField * StreamMemory.PacketsPerChunk * StreamMemory.PayloadSize;
@@ -950,23 +972,23 @@ namespace UE::RivermaxCore::Private
 			HeaderMemory->mkey[1] = InvalidKey;
 		}
 	
-		uint64 TotalSize = 0;
+		uint64 ProcessedBytes = 0;
 		uint64 LineSize = 0;
 		for (int32 PayloadSizeIndex = 0; PayloadSizeIndex < RealPacketsPerFrame; ++PayloadSizeIndex)
 		{
 			uint32 HeaderSize = FRawRTPHeader::OneSRDSize;
 			uint32 ThisPayloadSize = StreamMemory.PayloadSize;
-			if (TotalSize < FrameSize)
+			if (ProcessedBytes < FrameSizeInBytes)
 			{
-				if ((LineSize + StreamMemory.PayloadSize) == BytesPerLine)
+				if ((LineSize + StreamMemory.PayloadSize) == BytesPerRow)
 				{
 					LineSize = 0;
 				}
-				else if ((LineSize + StreamMemory.PayloadSize) > BytesPerLine)
+				else if ((LineSize + StreamMemory.PayloadSize) > BytesPerRow)
 				{
 					HeaderSize = FRawRTPHeader::TwoSRDSize;
-					LineSize = StreamMemory.PayloadSize - (BytesPerLine - LineSize);
-					if (LineSize > BytesPerLine)
+					LineSize = StreamMemory.PayloadSize - (BytesPerRow - LineSize);
+					if (LineSize > BytesPerRow)
 					{
 						UE_LOG(LogRivermax, Warning, TEXT("Unsupported small resolution, %dx%d, needing more than 2 SRD to express"), Options.AlignedResolution.X, Options.AlignedResolution.Y);
 						return false;
@@ -978,9 +1000,11 @@ namespace UE::RivermaxCore::Private
 					LineSize += StreamMemory.PayloadSize;
 				}
 
-				if ((TotalSize + StreamMemory.PayloadSize) > FrameSize)
+				// This means that this is the header for the last packet and it is smaller than all the previous ones.
+				if ((ProcessedBytes + StreamMemory.PayloadSize) > FrameSizeInBytes)
 				{
 					HeaderSize = FRawRTPHeader::OneSRDSize;
+					ThisPayloadSize = FrameSizeInBytes - ProcessedBytes;
 				}
 			}
 			else
@@ -1002,7 +1026,7 @@ namespace UE::RivermaxCore::Private
 				RTPFiller.Update(PayloadSizeIndex);
 			}
 			
-			TotalSize += ThisPayloadSize;
+			ProcessedBytes += ThisPayloadSize;
 		}
 
 		// Verify memcopy config to make sure it works for current frame size / chunking
@@ -1028,12 +1052,10 @@ namespace UE::RivermaxCore::Private
 	{
 		NextFrame->LineNumber = 0;
 		NextFrame->PacketCounter = 0;
-		NextFrame->SRDOffset = 0;
 		NextFrame->ChunkNumber = 0;
 		NextFrame->PayloadPtr = nullptr;
 		NextFrame->HeaderPtr = nullptr;
 		NextFrame->FrameStartPtr = nullptr;
-		NextFrame->ScheduleTimeCopied = 0;
 		NextFrame->bCaughtTimingIssue = false;
 		NextFrame->Offset = 0;
 	}
@@ -1042,6 +1064,8 @@ namespace UE::RivermaxCore::Private
 	{
 		using namespace UE::RivermaxCore::Private::Utils;
 
+		// Since the static part of the RTP header is filled on initalization we only need to update the 
+		// non static parts.
 		if (CachedCVars.bPrefillRTPHeaders)
 		{
 			OutHeader.SequenceNumber = ByteSwap((uint16)(StreamData.SequenceNumber & 0xFFFF));
@@ -1050,71 +1074,18 @@ namespace UE::RivermaxCore::Private
 		}
 		else
 		{
-			OutHeader = {};
-			OutHeader.Version = 2;
-			OutHeader.PaddingBit = 0;
-			OutHeader.ExtensionBit = 0;
-			OutHeader.PayloadType = 96; //Payload type should probably be infered from SDP
-
-			OutHeader.SequenceNumber = ByteSwap((uint16)(StreamData.SequenceNumber & 0xFFFF));
-			OutHeader.Timestamp = ByteSwap(CurrentFrame->MediaTimestamp);
-
-			//2110 specific header
-			OutHeader.SynchronizationSource = ByteSwap((uint32)0x0eb51dbd);  // Should Unreal has its own synch source ID
-
-			if (StreamType == ERivermaxStreamType::VIDEO_2110_20_STREAM)
-			{
-				if (CurrentFrame->PacketCounter + 1 == StreamMemory.PacketsPerFrame)
-				{
-					OutHeader.MarkerBit = 1; // last packet in frame (Marker)
-				}
-
-				OutHeader.ExtendedSequenceNumber = ByteSwap((uint16)((StreamData.SequenceNumber >> 16) & 0xFFFF));
-
-				// Verify if payload size exceeds line 
-				const uint32 CurrentPayloadSize = StreamMemory.PayloadSizes[CurrentFrame->PacketCounter];
-
-				const uint32 LineSizeOffset = ((CurrentFrame->SRDOffset / FormatInfo.PixelGroupCoverage) * FormatInfo.PixelGroupSize);
-				const uint32 LineSize = ((Options.AlignedResolution.X / FormatInfo.PixelGroupCoverage) * FormatInfo.PixelGroupSize);
-
-				const uint16 SRD1Length = FMath::Min(LineSize - LineSizeOffset, CurrentPayloadSize);
-				const uint16 SRD1PixelCount = SRD1Length / FormatInfo.PixelGroupSize * FormatInfo.PixelGroupCoverage;
-				uint16 SRD2Length = SRD1Length < CurrentPayloadSize ? CurrentPayloadSize - SRD1Length : 0;
-				if (SRD2Length && CurrentFrame->LineNumber == ((uint32)Options.AlignedResolution.Y - 1))
-				{
-					SRD2Length = 0;
-				}
-
-				OutHeader.SRD1Length = ByteSwap(SRD1Length);
-				OutHeader.SetSrd1RowNumber(CurrentFrame->LineNumber); //todo divide by 2 if interlaced
-				OutHeader.FieldIdentification1 = 0; //todo when fields are sent for interlace
-				OutHeader.SetSrd1Offset(CurrentFrame->SRDOffset);
-
-				CurrentFrame->SRDOffset += SRD1PixelCount;
-				if (CurrentFrame->SRDOffset >= Options.AlignedResolution.X)
-				{
-					CurrentFrame->SRDOffset = 0;
-					++CurrentFrame->LineNumber;
-				}
-
-				if (SRD2Length > 0)
-				{
-					OutHeader.SRD2Length = ByteSwap(SRD2Length);
 
-					OutHeader.ContinuationBit1 = 1;
-					OutHeader.FieldIdentification2 = 0;
-					OutHeader.SetSrd2RowNumber(CurrentFrame->LineNumber);
-					OutHeader.SetSrd2Offset(CurrentFrame->SRDOffset);
-
-					const uint16 SRD2PixelCount = SRD2Length / FormatInfo.PixelGroupSize * FormatInfo.PixelGroupCoverage;
-					CurrentFrame->SRDOffset += SRD2PixelCount;
-					if (CurrentFrame->SRDOffset >= Options.AlignedResolution.X)
-					{
-						CurrentFrame->SRDOffset = 0;
-						++CurrentFrame->LineNumber;
-					}
-				}
-			}
+			UpdateVideoRtpHeader(
+				OutHeader,
+				CurrentFrame->PacketCounter,
+				CurrentFrame->SRDOffset,
+				CurrentFrame->LineNumber,
+				StreamData.SequenceNumber,
+				CurrentFrame->MediaTimestamp,
+				Options.AlignedResolution,
+				FormatInfo,
+				StreamMemory
+			);
 		}
 	}
 
@@ -1184,8 +1155,7 @@ namespace UE::RivermaxCore::Private
 			// If a frame is ready already, we can move on. Otherwise, we wait for FrameReady signal with a
 			// wait timeout. In the case of a repeated frame, we will always timeout and we won't be able to 
 			// do an early copy.
-			TSharedPtr<FRivermaxOutputFrame> ReadyFrame = FrameManager->GetReadyFrame();
-			if (ReadyFrame)
+			if (FrameManager->IsFrameAvailableToSend())
 			{
 				bIsFrameReady = true;
 			}
@@ -1195,8 +1165,7 @@ namespace UE::RivermaxCore::Private
 				do 
 				{
 					bIsFrameReady = FrameReadyToSendSignal->Wait(WaitMs);
-					ReadyFrame = FrameManager->GetReadyFrame();
-				} while (!ReadyFrame && bIsFrameReady && bIsActive);
+				} while (!FrameManager->IsFrameAvailableToSend() && bIsFrameReady && bIsActive);
 			}
 		}
 		else
@@ -1268,11 +1237,6 @@ namespace UE::RivermaxCore::Private
 				{
 					// Stamp frame start in order to copy frame data sequentially as we query chunks
 					CurrentFrame->FrameStartPtr = CurrentFrame->PayloadPtr;
-
-					if (CurrentFrame->PayloadPtr != StreamMemory.BufferAddresses[StreamData.ExpectedFrameIndex])
-					{
-						UE_LOG(LogRivermax, Warning, TEXT("Frame being sent (%d) doesn't match chunks being processed."), CurrentFrame->FrameIndex);
-					}
 				}
 
 				break;
@@ -1300,35 +1264,30 @@ namespace UE::RivermaxCore::Private
 	bool FRivermaxOutputStream::CopyFrameData(const TSharedPtr<FRivermaxOutputFrame>& SourceFrame, uint8* DestinationBase)
 	{
 		// Make sure copy size doesn't go over frame size
-		const uint32 FrameSize = GetStride() * Options.AlignedResolution.Y;
+		const SIZE_T FrameSize = GetRowSizeInBytes() * Options.AlignedResolution.Y;
 		if (ensure(FrameSize > 0))
 		{
-			const uint32 BlockSize = 1 + ((FrameSize - 1) / StreamMemory.FrameMemorySliceCount);
-			const uint32 MaxSize = FrameSize - SourceFrame->Offset;
-			const uint32 CopySize = FMath::Min(BlockSize, MaxSize);
+			const SIZE_T BlockSize = 1 + ((FrameSize - 1) / StreamMemory.FrameMemorySliceCount);
+			const SIZE_T MaxSize = FrameSize - SourceFrame->Offset;
+			const SIZE_T CopySize = FMath::Min(BlockSize, MaxSize);
 
 			// Copy data until we have covered the whole frame. Last block might be smaller.
 			if (CopySize > 0)
 			{
 				TRACE_CPUPROFILER_EVENT_SCOPE(CopyFrameData);
-				const uint32 PacketsToCopy = CopySize / StreamMemory.PayloadSize;
-				const uint32 PacketOffset = SourceFrame->ChunkNumber * StreamMemory.PacketsPerChunk;
+				const SIZE_T PacketsToCopy = CopySize / StreamMemory.PayloadSize;
+				const SIZE_T PacketOffset = SourceFrame->ChunkNumber * StreamMemory.PacketsPerChunk;
 				const uint64 PacketDeltaTime = (StreamData.FrameFieldTimeIntervalNs - TransmitOffsetNanosec) / StreamMemory.PacketsPerFrame;
 
 				// Time covered in frame interval including this copy
 				const uint64 TimeCopied = StreamData.NextScheduleTimeNanosec + ((PacketOffset + PacketsToCopy) * PacketDeltaTime);
 
-				// Sidecar used to track copy progress
-				TSharedPtr<FMemChunkCopiedInfo> CopyInfo = MakeShared<FMemChunkCopiedInfo>();
-				CopyInfo->TimeCovered = TimeCopied;
-
 				FCopyArgs Args;
 				uint8* SourceStart = reinterpret_cast<uint8*>(SourceFrame->VideoBuffer);
 				uint8* DestinationStart = DestinationBase;
 				Args.SourceMemory = SourceStart + SourceFrame->Offset;
 				Args.DestinationMemory = DestinationStart + SourceFrame->Offset;
 				Args.SizeToCopy = CopySize;
-				Args.SideCar = MoveTemp(CopyInfo);
 
 				// Update memory offset for next copy
 				SourceFrame->Offset += CopySize;
@@ -1353,21 +1312,12 @@ namespace UE::RivermaxCore::Private
 		FRawRTPHeader* HeaderRawPtr = reinterpret_cast<FRawRTPHeader*>(CurrentFrame->HeaderPtr);
 		check(HeaderRawPtr);
 
-		const uint32 MemblockIndex = StreamData.ExpectedFrameIndex / StreamMemory.FramesFieldPerMemoryBlock;
-		const uint32 MemblockFrameIndex = StreamData.ExpectedFrameIndex % StreamMemory.FramesFieldPerMemoryBlock;
-		const uint32 PacketsPerFrame = StreamMemory.ChunksPerFrameField * StreamMemory.PacketsPerChunk;
-		const uint32 BaseHeaderIndex = CurrentFrame->PacketCounter + (MemblockFrameIndex * PacketsPerFrame);
-
-		if (ensure(HeaderRawPtr == &StreamMemory.RTPHeaders[MemblockIndex][BaseHeaderIndex]))
+		for (uint32 PacketIndex = 0; PacketIndex < StreamMemory.PacketsPerChunk && CurrentFrame->PacketCounter < StreamMemory.PacketsPerFrame; ++PacketIndex)
 		{
-			for (uint32 PacketIndex = 0; PacketIndex < StreamMemory.PacketsPerChunk && CurrentFrame->PacketCounter < StreamMemory.PacketsPerFrame; ++PacketIndex)
-			{
-				BuildRTPHeader(StreamMemory.RTPHeaders[MemblockIndex][BaseHeaderIndex + PacketIndex]);
+			BuildRTPHeader(*(HeaderRawPtr + PacketIndex));
 				
-				CurrentFrame->BytesSent += StreamMemory.PayloadSizes[CurrentFrame->PacketCounter];
-				++StreamData.SequenceNumber;
-				++CurrentFrame->PacketCounter;
-			}
+			++StreamData.SequenceNumber;
+			++CurrentFrame->PacketCounter;
 		}
 	}
 
@@ -1423,7 +1373,7 @@ namespace UE::RivermaxCore::Private
 		if (bIsActive && CurrentFrame->ChunkNumber == 0 && CachedCVars.bShowOutputStats)
 		{
 			UE_LOG(LogRivermax, Verbose, TEXT("Committed frame [%u]. Scheduled for '%llu'. Aligned with '%llu'. Current time '%llu'. Was late: %d. Slack: %llu. Errorcount: %d")
-				, CurrentFrame->FrameIdentifier
+				, CurrentFrame->GetFrameCounter()
 				, ScheduleTime
 				, StreamData.NextAlignmentPointNanosec
 				, CurrentTimeNanosec
@@ -1499,11 +1449,11 @@ namespace UE::RivermaxCore::Private
 	{
 		// When aligning on frame creation, we will always wait for a frame to be available.
 		TRACE_CPUPROFILER_EVENT_SCOPE(Rmax::WaitForReadyFrame);
-		TSharedPtr<FRivermaxOutputFrame> NextFrameToSend = FrameManager->GetReadyFrame();
-		while (!NextFrameToSend && bIsActive)
+		TSharedPtr<FRivermaxOutputFrame> NextFrameToSend = FrameManager->DequeueFrameToSend();
+		while (!NextFrameToSend.IsValid() && bIsActive)
 		{
 			FrameReadyToSendSignal->Wait();
-			NextFrameToSend = FrameManager->GetReadyFrame();
+			NextFrameToSend = FrameManager->DequeueFrameToSend();
 		}
 
 		// In frame creation alignment, we always release the last frame sent
@@ -1517,8 +1467,6 @@ namespace UE::RivermaxCore::Private
 		if (NextFrameToSend)
 		{
 			CurrentFrame = MoveTemp(NextFrameToSend);
-			FrameManager->MarkAsSending(CurrentFrame);
-
 			InitializeNextFrame(CurrentFrame);
 		}
 	}
@@ -1534,7 +1482,7 @@ namespace UE::RivermaxCore::Private
 		// if none are available and bDoContinuousOutput == false
 		//		Don't send a frame and go back waiting for the next alignment point
 		
-		TSharedPtr<FRivermaxOutputFrame> NextFrameToSend = FrameManager->GetReadyFrame();
+		TSharedPtr<FRivermaxOutputFrame> NextFrameToSend = FrameManager->DequeueFrameToSend();
 
 		// If we have a new frame, release the previous one
 		// If we don't have a frame and we're not doing continuous output, we release it and we won't send a new one
@@ -1551,7 +1499,6 @@ namespace UE::RivermaxCore::Private
 			if (NextFrameToSend)
 			{
 				CurrentFrame = MoveTemp(NextFrameToSend);
-				FrameManager->MarkAsSending(CurrentFrame);
 				InitializeNextFrame(CurrentFrame);
 			}
 		}
@@ -1569,13 +1516,10 @@ namespace UE::RivermaxCore::Private
 			if (!StreamMemory.bUseIntermediateBuffer)
 			{
 				// No frame to send, keep last one and restart its internal counters
-				UE_LOG(LogRivermax, Verbose, TEXT("No frame to send. Reusing last frame '%d' with Id %u"), CurrentFrame->FrameIndex, CurrentFrame->FrameIdentifier);
+				UE_LOG(LogRivermax, Verbose, TEXT("No frame to send. Reusing last frame with Frame Counter: %u"), CurrentFrame->GetFrameCounter());
 
 				// Since we want to resend last frame, we need to fast forward chunk pointer to re-point to the one we just sent
 				SkipChunks(StreamMemory.ChunksPerFrameField * (Options.NumberOfBuffers - 1));
-				
-				// Keep RTP header in sync when skipping chunks
-				StreamData.ExpectedFrameIndex = CurrentFrame->FrameIndex;
 			}
 		}
 	}
@@ -1704,7 +1648,7 @@ namespace UE::RivermaxCore::Private
 		FrameManagerArgs.Resolution = Options.AlignedResolution;
 		FrameManagerArgs.bTryGPUAllocation = bTryGPUDirect;
 		FrameManagerArgs.NumberOfFrames = Options.NumberOfBuffers;
-		FrameManagerArgs.Stride = GetStride();
+		FrameManagerArgs.Stride = GetRowSizeInBytes();
 		FrameManagerArgs.FrameDesiredSize = StreamMemory.ChunksPerFrameField * StreamMemory.PacketsPerChunk * StreamMemory.PayloadSize;
 		FrameManagerArgs.bAlignEachFrameAlloc = bAlignEachFrameMemory;
 		FrameManagerArgs.OnFreeFrameDelegate = FOnFrameReadyDelegate::CreateRaw(this, &FRivermaxOutputStream::OnFrameReadyToBeUsed);
@@ -1758,10 +1702,16 @@ namespace UE::RivermaxCore::Private
 		}
 		else
 		{
+			// Used to keep track of acquired frames to prevent the same frame from being acquired from the pool more than once. 
+			// When this scope is exited all frames are returned back to the Frame Manager pool.
+			TArray<TSharedPtr<UE::RivermaxCore::Private::FRivermaxOutputFrame>> ProcessedFrames;
+
 			// When we don't use intermediate buffer, each frame has its own buffer addresses and we don't need to look at memblocks
 			for (int32 BufferIndex = 0; BufferIndex < Options.NumberOfBuffers; ++BufferIndex)
 			{
-				StreamMemory.BufferAddresses.Add(FrameManager->GetFrame(BufferIndex)->VideoBuffer);
+				TSharedPtr<UE::RivermaxCore::Private::FRivermaxOutputFrame> Frame = FrameManager->GetFreeFrame();
+				StreamMemory.BufferAddresses.Add(Frame->VideoBuffer);
+				ProcessedFrames.Add(Frame);
 			}
 		}
 		check(StreamMemory.BufferAddresses.Num() == Options.NumberOfBuffers);
@@ -1784,7 +1734,7 @@ namespace UE::RivermaxCore::Private
 		}
 	}
 
-	int32 FRivermaxOutputStream::GetStride() const
+	SIZE_T FRivermaxOutputStream::GetRowSizeInBytes() const
 	{
 		check(FormatInfo.PixelGroupCoverage != 0);
 		return (Options.AlignedResolution.X / FormatInfo.PixelGroupCoverage) * FormatInfo.PixelGroupSize;
@@ -1884,20 +1834,26 @@ namespace UE::RivermaxCore::Private
 		}
 	}
 
-	bool FRivermaxOutputStream::ReserveFrame(uint32 FrameIdentifier) const
+	bool FRivermaxOutputStream::ReserveFrame()
 	{
-		TSharedPtr<FRivermaxOutputFrame> ReservedFrame = FrameManager->GetNextFrame(FrameIdentifier);
-		if (!ReservedFrame && Options.FrameLockingMode == EFrameLockingMode::BlockOnReservation)
+		// There is only one reserved frame at the time per stream.
+		TSharedPtr<FRivermaxOutputFrame> ReservedFrame = FrameManager->GetFreeFrame();
+		if (!ReservedFrame.IsValid() && Options.FrameLockingMode == EFrameLockingMode::BlockOnReservation)
 		{
 			TRACE_CPUPROFILER_EVENT_SCOPE(Rmax::WaitForAvailableFrame);
 			while(!ReservedFrame && bIsActive)
 			{
 				FrameAvailableSignal->Wait();
-				ReservedFrame = FrameManager->GetNextFrame(FrameIdentifier);
+				ReservedFrame = FrameManager->GetFreeFrame();
 			}
 		}
 
-		return ReservedFrame != nullptr;
+		if (ReservedFrame.IsValid())
+		{
+			ReservedFrames.Enqueue(ReservedFrame);
+		}
+
+		return ReservedFrame.IsValid();
 	}
 
 	void FRivermaxOutputStream::OnFrameReadyToBeSent()
@@ -1984,14 +1940,6 @@ namespace UE::RivermaxCore::Private
 
 	void FRivermaxOutputStream::OnMemoryChunksCopied(const TSharedPtr<FBaseDataCopySideCar>& Sidecar)
 	{
-		if (Sidecar)
-		{
-			TSharedPtr<FMemChunkCopiedInfo> CopyInfo = StaticCastSharedPtr<FMemChunkCopiedInfo>(Sidecar);
-			if (CurrentFrame)
-			{
-				CurrentFrame->ScheduleTimeCopied = CopyInfo->TimeCovered;
-			}
-		}
 	}
 
 	void FRivermaxOutputStream::OnCVarRandomDelayChanged(IConsoleVariable* Var)
@@ -2008,7 +1956,7 @@ namespace UE::RivermaxCore::Private
 		uint64 InputTime = StreamData.NextAlignmentPointNanosec;
 		if (Options.bDoFrameCounterTimestamping)
 		{
-			InputTime = UE::RivermaxCore::GetAlignmentPointFromFrameNumber(CurrentFrame->FrameIdentifier, Options.FrameRate);
+			InputTime = UE::RivermaxCore::GetAlignmentPointFromFrameNumber(CurrentFrame->GetFrameCounter(), Options.FrameRate);
 		}
 
 		CurrentFrame->MediaTimestamp = GetTimestampFromTime(InputTime, MediaClockSampleRate);
@@ -2062,7 +2010,8 @@ namespace UE::RivermaxCore::Private
 
 		const uint32 MediaFrameNumber = Utils::TimestampToFrameNumber(CurrentFrame->MediaTimestamp, Options.FrameRate);
 		TRACE_CPUPROFILER_EVENT_SCOPE_TEXT(*FRivermaxTracingUtils::RmaxOutSendingFrameTraceEvents[MediaFrameNumber % 10]);
-		TRACE_CPUPROFILER_EVENT_SCOPE_TEXT(*FRivermaxTracingUtils::RmaxOutMediaCapturePipeTraceEvents[CurrentFrame->FrameIdentifier % 10]);
+
+		TRACE_CPUPROFILER_EVENT_SCOPE_TEXT(*FRivermaxTracingUtils::RmaxOutMediaCapturePipeTraceEvents[CurrentFrame->GetFrameCounter() % 10]);
 		UE_LOG(LogRivermax, VeryVerbose, TEXT("RmaxRX Sending frame number %u with timestamp %u."), MediaFrameNumber, CurrentFrame->MediaTimestamp);
 
 		do
@@ -2119,30 +2068,6 @@ namespace UE::RivermaxCore::Private
 		} while (CurrentFrame->ChunkNumber < StreamMemory.ChunksPerFrameField && bIsActive);
 
 		Stats.FramesSentCounter++;
-		StreamData.ExpectedFrameIndex = (StreamData.ExpectedFrameIndex + 1) % Options.NumberOfBuffers;
-	}
-
-	void FRivermaxOutputStream::InitializeTimingProtections()
-	{
-		switch (Options.AlignmentMode)
-		{
-			case ERivermaxAlignmentMode::FrameCreation:
-			{
-				StreamMemory.bAlwaysSkipChunk = false;
-				StreamMemory.FrameRateMultiplier = 1.0f;
-				break;
-			}
-			case ERivermaxAlignmentMode::AlignmentPoint:
-			{
-				StreamMemory.bAlwaysSkipChunk = CVarRivermaxOutputForceSkip.GetValueOnAnyThread();
-				StreamMemory.FrameRateMultiplier = FMath::Clamp(CVarRivermaxOutputFrameRateMultiplier.GetValueOnAnyThread(), 1.0, 2.0);
-				break;
-			}
-			default:
-			{
-				checkNoEntry();
-			}
-		}
 	}
 
 	void FRivermaxOutputStream::GetLastPresentedFrame(FPresentedFrameInfo& OutFrameInfo) const
@@ -2157,15 +2082,15 @@ namespace UE::RivermaxCore::Private
 		{
 			{
 				FScopeLock Lock(&PresentedFrameCS);
-				LastPresentedFrame.RenderedFrameNumber = CurrentFrame->FrameIdentifier;
 				LastPresentedFrame.PresentedFrameBoundaryNumber = StreamData.LastAlignmentPointFrameNumber;
+				LastPresentedFrame.RenderedFrameNumber = CurrentFrame->GetFrameCounter();
 			}
 			
 			// We don't release when there is no new frame, so we keep a hold on it to repeat it.
 			if (bReleaseFrame)
 			{
-				FrameManager->MarkAsSent(CurrentFrame);
 				CurrentFrame.Reset();
+				FrameManager->FrameSentEvent();
 			}
 		}
 	}
diff --git a/Engine/Plugins/VirtualProduction/Rivermax/RivermaxCore/Source/RivermaxCore/Private/RivermaxOutputStream.h b/Engine/Plugins/VirtualProduction/Rivermax/RivermaxCore/Source/RivermaxCore/Private/RivermaxOutputStream.h
--- a/Engine/Plugins/VirtualProduction/Rivermax/RivermaxCore/Source/RivermaxCore/Private/RivermaxOutputStream.h
+++ b/Engine/Plugins/VirtualProduction/Rivermax/RivermaxCore/Source/RivermaxCore/Private/RivermaxOutputStream.h
@@ -80,19 +80,6 @@ namespace UE::RivermaxCore::Private
 		/** Chunk committed between each memcopy of frame data. Helps respect timing. */
 		uint32 ChunkSpacingBetweenMemcopies = 1;
 
-		/** 
-		 * Whether we skip zero chunk, to reset rivermax's stream internal timings, vblank verification, etc...
-		 * Alignment point method will default to true and be configurable through cvar and frame creation will never use it.
-		 * Usage is to make sure Tro is respected when frame rate multiplier is involved.
-		 */
-		bool bAlwaysSkipChunk = true;
-
-		/** 
-		 * Multiplier applied to desired frame rate when creating Rivermax stream to add margin in case of packet timings not being respected
-		 * This will break compliance with standard but if receiver can tolerate it, more stable stream timings are expected.
-		 */
-		float FrameRateMultiplier = 1.00;
-
 		/** Memory blocks given to Rivermax which where data is located */
 		TArray<rmx_output_media_mem_block> MemoryBlocks;
 
@@ -150,9 +137,6 @@ namespace UE::RivermaxCore::Private
 		/** Current sequence number being done */
 		uint32 SequenceNumber = 0;
 		double FrameFieldTimeIntervalNs = 0.0;
-		
-		/** Data and RTP frame index expected to be used for next frame */
-		uint8 ExpectedFrameIndex = 0;
 
 		/** Next alignment point based on PTP standard */
 		uint64 NextAlignmentPointNanosec = 0;
@@ -221,7 +205,7 @@ namespace UE::RivermaxCore::Private
 		virtual void Uninitialize() override;
 		virtual bool PushVideoFrame(const FRivermaxOutputVideoFrameInfo& NewFrame) override;
 		virtual bool IsGPUDirectSupported() const override;
-		virtual bool ReserveFrame(uint32 FrameIdentifier) const override;
+		virtual bool ReserveFrame() override;
 		virtual void GetLastPresentedFrame(FPresentedFrameInfo& OutFrameInfo) const override;
 		//~ End IRivermaxOutputStream interface
 
@@ -242,9 +226,6 @@ namespace UE::RivermaxCore::Private
 		/** Initializes timing setup for this stream. TRO, frame interval etc... */
 		void InitializeStreamTimingSettings();
 
-		/** Configures settings related to timing protection dependant on cvars */
-		void InitializeTimingProtections();
-
 		/** Sets up frame management taking care of allocation, special cuda handling, etc... */
 		bool SetupFrameManagement();
 
@@ -311,7 +292,7 @@ namespace UE::RivermaxCore::Private
 		uint32 GetTimestampFromTime(uint64 InTimeNanosec, double InMediaClockRate) const;
 		
 		/** Get row stride for the current stream configuration */
-		int32 GetStride() const;
+		SIZE_T GetRowSizeInBytes() const;
 
 		/** Used to notify the listener that a frame is ready to be enqueued for transmission */
 		void OnPreFrameReadyToBeSent();
@@ -366,6 +347,11 @@ namespace UE::RivermaxCore::Private
 		/** Current frame being sent */
 		TSharedPtr<FRivermaxOutputFrame> CurrentFrame;
 
+		/** 
+		* Frames reserved when Media Capture texture is about to be converted into Buffer to be sent.
+		*/
+		TQueue<TSharedPtr<FRivermaxOutputFrame>> ReservedFrames;
+
 		/** Thread scheduling frame output */
 		TUniquePtr<FRunnableThread> RivermaxThread;
 
@@ -387,7 +373,7 @@ namespace UE::RivermaxCore::Private
 		/** Listener for this stream events */
 		IRivermaxOutputStreamListener* Listener = nullptr;
 
-		/** Type of stream created. Only 21110-20 (Video is supported now) */
+		/** Type of stream created. Only 2110-20 (Video is supported now) */
 		ERivermaxStreamType StreamType = ERivermaxStreamType::VIDEO_2110_20_STREAM;
 
 		/** TRoffset time calculated based on ST2110 - 21 Gapped(for now) method. This is added to next alignment point */
diff --git a/Engine/Plugins/VirtualProduction/Rivermax/RivermaxCore/Source/RivermaxCore/Private/RivermaxUtils.cpp b/Engine/Plugins/VirtualProduction/Rivermax/RivermaxCore/Source/RivermaxCore/Private/RivermaxUtils.cpp
--- a/Engine/Plugins/VirtualProduction/Rivermax/RivermaxCore/Source/RivermaxCore/Private/RivermaxUtils.cpp
+++ b/Engine/Plugins/VirtualProduction/Rivermax/RivermaxCore/Source/RivermaxCore/Private/RivermaxUtils.cpp
@@ -87,7 +87,7 @@ namespace UE::RivermaxCore::Private::Utils
 		}
 	}
 
-	void StreamOptionsToSDPDescription(const UE::RivermaxCore::FRivermaxOutputStreamOptions& Options, float RateMultiplier, FAnsiStringBuilderBase& OutSDPDescription)
+	void StreamOptionsToSDPDescription(const UE::RivermaxCore::FRivermaxOutputStreamOptions& Options, FAnsiStringBuilderBase& OutSDPDescription)
 	{
 		// Basic SDP string creation from a set of options. At some point, having a proper SDP loader / creator would be useful.
 		// Refer to https://datatracker.ietf.org/doc/html/rfc4570
@@ -102,9 +102,9 @@ namespace UE::RivermaxCore::Private::Utils
 		}
 
 		FString FrameRateDescription;
-		if (!FMath::IsNearlyEqual(RateMultiplier, 1.0f) || bIsNonStandardFractionalFrameRate)
+		if (bIsNonStandardFractionalFrameRate)
 		{
-			const double NewRateDecimal = DecimalDesiredRate * RateMultiplier * 1001;
+			const double NewRateDecimal = DecimalDesiredRate * 1001;
 			DesiredRate.Numerator = FMath::RoundToInt32(NewRateDecimal);
 			DesiredRate.Denominator = 1001;
 		}
diff --git a/Engine/Plugins/VirtualProduction/Rivermax/RivermaxCore/Source/RivermaxCore/Private/RivermaxUtils.h b/Engine/Plugins/VirtualProduction/Rivermax/RivermaxCore/Source/RivermaxCore/Private/RivermaxUtils.h
--- a/Engine/Plugins/VirtualProduction/Rivermax/RivermaxCore/Source/RivermaxCore/Private/RivermaxUtils.h
+++ b/Engine/Plugins/VirtualProduction/Rivermax/RivermaxCore/Source/RivermaxCore/Private/RivermaxUtils.h
@@ -31,7 +31,7 @@ namespace UE::RivermaxCore::Private::Utils
 	static constexpr float SleepTimeSeconds = 50 * 1E-6;
 
 	/** Convert a set of streaming option to its SDP description. Currently only support video type. */
-	void StreamOptionsToSDPDescription(const UE::RivermaxCore::FRivermaxOutputStreamOptions& Options, float RateMultiplier, FAnsiStringBuilderBase& OutSDPDescription);
+	void StreamOptionsToSDPDescription(const UE::RivermaxCore::FRivermaxOutputStreamOptions& Options, FAnsiStringBuilderBase& OutSDPDescription);
 
 	/**
 	 * Converts a timestamp in MediaClock period units to a frame number for a given frame rate
diff --git a/Engine/Plugins/VirtualProduction/Rivermax/RivermaxCore/Source/RivermaxCore/Public/IRivermaxOutputStream.h b/Engine/Plugins/VirtualProduction/Rivermax/RivermaxCore/Source/RivermaxCore/Public/IRivermaxOutputStream.h
--- a/Engine/Plugins/VirtualProduction/Rivermax/RivermaxCore/Source/RivermaxCore/Public/IRivermaxOutputStream.h
+++ b/Engine/Plugins/VirtualProduction/Rivermax/RivermaxCore/Source/RivermaxCore/Public/IRivermaxOutputStream.h
@@ -26,9 +26,13 @@ namespace UE::RivermaxCore
 		/** Stride of a line */
 		uint32 Stride = 0;
 
+		UE_DEPRECATED(5.6, "This property is deprecated. Please use CPUBuffer instead")
 		/** Video frame data location in system memory */
 		void* VideoBuffer = nullptr;
 
+		/** Video frame data location in system memory */
+		void* CPUBuffer = nullptr;
+
 		/** Video frame data location when using GPUDirect */
 		FBufferRHIRef GPUBuffer;
 	};
@@ -75,13 +79,6 @@ namespace UE::RivermaxCore
 
 		/** Pushes new video frame to the stream. Returns false if frame couldn't be pushed. */
 		virtual bool PushVideoFrame(const FRivermaxOutputVideoFrameInfo& NewFrame) = 0;
-		
-		/** Pushes a frame to be captured from GPU memory */
-		UE_DEPRECATED("5.3", "Use PushVideoFrame instead and fill GPUBuffer with Capturedbuffer")
-		virtual bool PushGPUVideoFrame(const FRivermaxOutputVideoFrameInfo& NewFrame, FBufferRHIRef CapturedBuffer) final
-		{
-			return false;
-		}
 
 		/** Returns true if GPUDirect is supported */
 		virtual bool IsGPUDirectSupported() const = 0;
@@ -91,7 +88,15 @@ namespace UE::RivermaxCore
 		 * If FRivermaxOutputStreamOptions::EFrameLockingMode equals BlockOnReservation, 
 		 * this method will block until a free frame is found. 
 		 */
-		virtual bool ReserveFrame(uint32 FrameIdentifier) const = 0;
+		UE_DEPRECATED(5.6, "This method is deprecated. Please use the ReserveFrame that doesn't rely on frame identifier.")
+		virtual bool ReserveFrame(uint32 FrameIdentifier) const { return false; };
+
+		/** 
+		 * Tries to reserve a frame for the next capture 
+		 * If FRivermaxOutputStreamOptions::EFrameLockingMode equals BlockOnReservation, 
+		 * this method will block until a free frame is found. 
+		 */
+		virtual bool ReserveFrame() = 0;
 
 		/** Returns information about the last frame that was presented on the wire */
 		virtual void GetLastPresentedFrame(FPresentedFrameInfo& OutFrameInfo) const = 0;
diff --git a/Engine/Plugins/VirtualProduction/Rivermax/RivermaxMedia/Source/RivermaxMedia/Private/RivermaxMediaCapture.cpp b/Engine/Plugins/VirtualProduction/Rivermax/RivermaxMedia/Source/RivermaxMedia/Private/RivermaxMediaCapture.cpp
--- a/Engine/Plugins/VirtualProduction/Rivermax/RivermaxMedia/Source/RivermaxMedia/Private/RivermaxMediaCapture.cpp
+++ b/Engine/Plugins/VirtualProduction/Rivermax/RivermaxMedia/Source/RivermaxMedia/Private/RivermaxMediaCapture.cpp
@@ -272,6 +272,10 @@ void URivermaxMediaCapture::AddFrameReservationPass(FRDGBuilder& GraphBuilder)
 {
 	RHI_BREADCRUMB_EVENT_STAT(GraphBuilder.RHICmdList, Rmax_FrameReservation, "Rmax_FrameReservation");
 	
+	// Since we are going to enqueue a lambda that can potentially sleep in the RHI thread if the pixels haven't arrived,
+	// we dispatch the existing commands before any potential sleep.
+	GraphBuilder.RHICmdList.ImmediateFlush(EImmediateFlushType::DispatchToRHIThread);
+
 	// Scene rendering will already be enqueued but capture conversion pass will not
 	// Revisit to push slot reservation till last minute
 	URivermaxMediaCapture* Capturer = this;
@@ -282,7 +286,7 @@ void URivermaxMediaCapture::AddFrameReservationPass(FRDGBuilder& GraphBuilder)
 
 		if (Capturer->bIsActive)
 		{
-			Capturer->RivermaxStream->ReserveFrame(FrameCounter);
+			Capturer->RivermaxStream->ReserveFrame();
 		}
 	});
 }
@@ -296,11 +300,11 @@ void URivermaxMediaCapture::OnFrameCapturedInternal_AnyThread(const FCaptureBase
 	NewFrame.Height = Height;
 	NewFrame.Width = Width;
 	NewFrame.Stride = BytesPerRow;
-	NewFrame.VideoBuffer = InBuffer;
+	NewFrame.CPUBuffer = InBuffer;
 	NewFrame.FrameIdentifier = InBaseData.SourceFrameNumberRenderThread;
 	if (RivermaxStream->PushVideoFrame(NewFrame) == false)
 	{
-		UE_LOG(LogRivermaxMedia, Verbose, TEXT("Failed to pushed captured frame"));
+		UE_LOG(LogRivermaxMedia, Verbose, TEXT("Failed to push captured frame"));
 	}
 }
 
