Rivermax: removed deprecated member which was causing issues in monolythics complication.

#rnx

diff --git a/Engine/Plugins/VirtualProduction/Rivermax/RivermaxCore/Source/RivermaxCore/Public/IRivermaxOutputStream.h b/Engine/Plugins/VirtualProduction/Rivermax/RivermaxCore/Source/RivermaxCore/Public/IRivermaxOutputStream.h
--- a/Engine/Plugins/VirtualProduction/Rivermax/RivermaxCore/Source/RivermaxCore/Public/IRivermaxOutputStream.h
+++ b/Engine/Plugins/VirtualProduction/Rivermax/RivermaxCore/Source/RivermaxCore/Public/IRivermaxOutputStream.h
@@ -26,10 +26,6 @@ namespace UE::RivermaxCore
 		/** Stride of a line */
 		uint32 Stride = 0;
 
-		UE_DEPRECATED(5.6, "This property is deprecated. Please use CPUBuffer instead")
-		/** Video frame data location in system memory */
-		void* VideoBuffer = nullptr;
-
 		/** Video frame data location in system memory */
 		void* CPUBuffer = nullptr;
 
