Rivermax Plugin Refactor:
- Factored out Video stream.
- Added core for Anc stream. (hidden from public)
- Added Export SDP option to Media Output (not in ndisplay configs yet)
- other minor improvements.

#rnx

diff --git a/Engine/Plugins/Media/MediaIOFramework/Source/MediaIOCore/Public/MediaCapture.h b/Engine/Plugins/Media/MediaIOFramework/Source/MediaIOCore/Public/MediaCapture.h
--- a/Engine/Plugins/Media/MediaIOFramework/Source/MediaIOCore/Public/MediaCapture.h
+++ b/Engine/Plugins/Media/MediaIOFramework/Source/MediaIOCore/Public/MediaCapture.h
@@ -477,8 +477,8 @@ protected:
 	{
 		FTimecode SourceFrameTimecode;
 		FFrameRate SourceFrameTimecodeFramerate;
-		uint32 SourceFrameNumberRenderThread = 0;
-		uint32 SourceFrameNumber = 0;
+		uint64 SourceFrameNumberRenderThread = 0;
+		uint64 SourceFrameNumber = 0;
 	};
 
 	
diff --git a/Engine/Plugins/VirtualProduction/Rivermax/RivermaxCore/Source/RivermaxCore/Private/RTPHeader.cpp b/Engine/Plugins/VirtualProduction/Rivermax/RivermaxCore/Source/RivermaxCore/Private/RTPHeader.cpp
--- a/Engine/Plugins/VirtualProduction/Rivermax/RivermaxCore/Source/RivermaxCore/Private/RTPHeader.cpp
+++ b/Engine/Plugins/VirtualProduction/Rivermax/RivermaxCore/Source/RivermaxCore/Private/RTPHeader.cpp
@@ -46,35 +46,35 @@ namespace UE::RivermaxCore::Private
 		return InHeader;
 	}
 
-	FRTPHeader::FRTPHeader(const FRawRTPHeader& RawHeader)
+	FRTPHeader::FRTPHeader(const FVideoRTPHeader& VideoRTP)
 	{
 		Timestamp = 0;
 
-		if (RawHeader.Version != 2)
+		if (VideoRTP.RTPHeader.Version != 2)
 		{
 			return;
 		}
 
 		// Pretty sure some data needs to be swapped but can't validate that until we have other hardware generating data
-		SequenceNumber = (ByteSwap((uint16)RawHeader.ExtendedSequenceNumber) << 16) | ByteSwap((uint16)RawHeader.SequenceNumber);
-		Timestamp = ByteSwap(RawHeader.Timestamp);
-		bIsMarkerBit = RawHeader.MarkerBit;
+		SequenceNumber = (ByteSwap((uint16)VideoRTP.RTPHeader.ExtendedSequenceNumber) << 16) | ByteSwap((uint16)VideoRTP.RTPHeader.SequenceNumber);
+		Timestamp = ByteSwap(VideoRTP.RTPHeader.Timestamp);
+		bIsMarkerBit = VideoRTP.RTPHeader.MarkerBit;
 
-		SyncSouceId = RawHeader.SynchronizationSource;
+		SyncSouceId = VideoRTP.RTPHeader.SynchronizationSource;
 
-		SRD1.Length = ByteSwap((uint16)RawHeader.SRD1.Length);
-		SRD1.DataOffset = RawHeader.SRD1.GetOffset();
-		SRD1.RowNumber = RawHeader.SRD1.GetRowNumber();
-		SRD1.bIsFieldOne = RawHeader.SRD1.FieldIdentification;
-		SRD1.bHasContinuation = RawHeader.SRD1.ContinuationBit;
+		SRD1.Length = ByteSwap((uint16)VideoRTP.SRD1.Length);
+		SRD1.DataOffset = VideoRTP.SRD1.GetOffset();
+		SRD1.RowNumber = VideoRTP.SRD1.GetRowNumber();
+		SRD1.bIsFieldOne = VideoRTP.SRD1.FieldIdentification;
+		SRD1.bHasContinuation = VideoRTP.SRD1.ContinuationBit;
 
 		if (SRD1.bHasContinuation)
 		{
-			SRD2.Length = ByteSwap((uint16)RawHeader.SRD2.Length);
-			SRD2.DataOffset = RawHeader.SRD2.GetOffset();
-			SRD2.RowNumber = RawHeader.SRD2.GetRowNumber();
-			SRD2.bIsFieldOne = RawHeader.SRD2.FieldIdentification;
-			SRD2.bHasContinuation = RawHeader.SRD2.ContinuationBit;
+			SRD2.Length = ByteSwap((uint16)VideoRTP.SRD2.Length);
+			SRD2.DataOffset = VideoRTP.SRD2.GetOffset();
+			SRD2.RowNumber = VideoRTP.SRD2.GetRowNumber();
+			SRD2.bIsFieldOne = VideoRTP.SRD2.FieldIdentification;
+			SRD2.bHasContinuation = VideoRTP.SRD2.ContinuationBit;
 
 			if (SRD2.bHasContinuation == true)
 			{
@@ -123,6 +123,56 @@ namespace UE::RivermaxCore::Private
 
 		return SRD1.RowNumber;
 	}
+
+	uint8 FAncillaryTimecodeHeaderFields::EvenParity(uint64 Num)
+	{
+		return (uint8)(FMath::CountBits(Num) & 0x1);
+	}
+
+	void FAncillaryTimecodeHeaderFields::SetATCTimecode(uint8 Hours, uint8 Minutes, uint8 Seconds, uint8 Frames, bool bDropFrame)
+	{
+		auto ConvertToBinaryCodedDecimal = [](uint8 OriginalNum) -> uint8 { return ((OriginalNum / 10) << 4) | (OriginalNum % 10); };
+
+		auto ParityWrap = [this](uint8 v) -> uint16
+			{
+				uint8 p = EvenParity(v);
+				return (v & 0xFF) | ((p ? 1 : 0) << 8) | ((p ? 0 : 1) << 9);
+			};
+
+		UserData1 = ParityWrap(ConvertToBinaryCodedDecimal(Frames) | (bDropFrame ? 0x40 : 0x00));  // Bit 7 = DF
+		UserData2 = ParityWrap(ConvertToBinaryCodedDecimal(Seconds));                              // Bit 6 = BGF1 if needed
+		UserData3 = ParityWrap(ConvertToBinaryCodedDecimal(Minutes));                              // Bit 6 = BGF2 if needed
+		UserData4 = ParityWrap(ConvertToBinaryCodedDecimal(Hours));                                // Bit 6 = BGF3 if needed
+		UserData5 = ParityWrap(0x00); // User bits 1
+		UserData6 = ParityWrap(0x00); // User bits 2
+		UserData7 = ParityWrap(0x00); // User bits 3
+		UserData8 = ParityWrap(0x00); // User bits 4
+
+		uint16 DID = 0x60;
+		uint16 SDID = 0x60;
+		uint16 DataCount = 8;
+
+		SetDataCount(DataCount); // Also sets parity/inverse parity bits
+
+		// Compute checksum
+		uint16 CheckSum = (DID & 0x1FF) + (SDID & 0x1FF) + (DataCount & 0x1FF)
+			+ (UserData1 & 0x1FF) + (UserData2 & 0x1FF) + (UserData3 & 0x1FF) + (UserData4 & 0x1FF)
+			+ (UserData5 & 0x1FF) + (UserData6 & 0x1FF) + (UserData7 & 0x1FF) + (UserData8 & 0x1FF);
+
+		CheckSum &= 0x1FF;
+
+		uint8 Lower8Bits = CheckSum & 0xFF;
+		uint8 NumOfSetBits = FMath::CountBits(Lower8Bits);
+		uint8 ParityBit = (NumOfSetBits % 2 == 0) ? 0 : 1;
+		uint8 InverseParityBit = ParityBit ^ 1;
+
+		CheckSum |= (ParityBit << 8);
+		CheckSum |= (InverseParityBit << 9);
+
+		ChecksumHigh = (CheckSum >> 2) & 0xFF;
+		ChecksumLow = CheckSum & 0xFF;
+
+	}
 }
 
 
diff --git a/Engine/Plugins/VirtualProduction/Rivermax/RivermaxCore/Source/RivermaxCore/Private/RTPHeader.h b/Engine/Plugins/VirtualProduction/Rivermax/RivermaxCore/Source/RivermaxCore/Private/RTPHeader.h
--- a/Engine/Plugins/VirtualProduction/Rivermax/RivermaxCore/Source/RivermaxCore/Private/RTPHeader.h
+++ b/Engine/Plugins/VirtualProduction/Rivermax/RivermaxCore/Source/RivermaxCore/Private/RTPHeader.h
@@ -19,16 +19,9 @@ namespace UE::RivermaxCore::Private
 	+-+-+-+-+-+-+-+-+-+-+-+-+-+-+-+-+-+-+-+-+-+-+-+-+-+-+-+-+-+-+-+-+
 	|                           ssrc                                |
 	+-+-+-+-+-+-+-+-+-+-+-+-+-+-+-+-+-+-+-+-+-+-+-+-+-+-+-+-+-+-+-+-+
-
-	SRD Header 8-14 bytes
-
-	0                   1                   2                   3
-	0 1 2 3 4 5 6 7 8 9 0 1 2 3 4 5 6 7 8 9 0 1 2 3 4 5 6 7 8 9 0 1
-	+-+-+-+-+-+-+-+-+-+-+-+-+-+-+-+-+-+-+-+-+-+-+-+-+-+-+-+-+-+-+-+-+
-	|    Extended Sequence Number   |           SRD Length          |
-	+-+-+-+-+-+-+-+-+-+-+-+-+-+-+-+-+-+-+-+-+-+-+-+-+-+-+-+-+-+-+-+-+
-	|F|     SRD Row Number          |C|         SRD Offset          |
-	+-+-+-+-+-+-+-+-+-+-+-+-+-+-+-+-+-+-+-+-+-+-+-+-+-+-+-+-+-+-+-+-+ */
+	|    Extended Sequence Number   |
+	+-+-+-+-+-+-+-+-+-+-+-+-+-+-+-+-+
+	*/
 
 	/** Raw representation as it's built for the network */
 
@@ -36,22 +29,50 @@ namespace UE::RivermaxCore::Private
 #if PLATFORM_SUPPORTS_PRAGMA_PACK
 #pragma pack(push, 1)
 #endif
-	/** Total packed size should be 6 octets. */
+	/** Total size should be 14 octets. */
+	struct FRawRTPHeader
+	{
+		uint8 ContributingSourceCount   : 4;
+		uint8 ExtensionBit              : 1;
+		uint8 PaddingBit                : 1;
+		uint8 Version                   : 2;
+		uint8 PayloadType               : 7;
+		uint8 MarkerBit                 : 1;
+		uint16 SequenceNumber           : 16;
+		uint32 Timestamp                : 32;
+		uint32 SynchronizationSource    : 32;
+		uint16 ExtendedSequenceNumber   : 16;
+
+		/** The base value is 0x0eb51dbf which will be used for video. Anc is VideoSyncSource + 1. */
+		static constexpr uint32 VideoSynchronizationSource = 0x0eb51dbf;
+	};
+
+	/** 
+	SRD Header. Total packed size should be 6 octets.
+
+	0                   1                   2                   3
+	0 1 2 3 4 5 6 7 8 9 0 1 2 3 4 5 6 7 8 9 0 1 2 3 4 5 6 7 8 9 0 1
+	+-+-+-+-+-+-+-+-+-+-+-+-+-+-+-+-+-+-+-+-+-+-+-+-+-+-+-+-+-+-+-+-+
+	|           SRD Length          |F|     SRD Row Number          |
+	+-+-+-+-+-+-+-+-+-+-+-+-+-+-+-+-+-+-+-+-+-+-+-+-+-+-+-+-+-+-+-+-+
+	|C|         SRD Offset          |
+	+-+-+-+-+-+-+-+-+-+-+-+-+-+-+-+-+
+	*/
 	struct FRawSRD
 	{
-		uint16 Length				: 16;
+		uint16 Length : 16;
 
-		uint8 RowNumberHigh			: 7;
+		uint8 RowNumberHigh : 7;
 
-		uint8 FieldIdentification	: 1;
+		uint8 FieldIdentification : 1;
 
-		uint8 RowNumberLow			: 8;
+		uint8 RowNumberLow : 8;
 
-		uint8 OffsetHigh			: 7;
+		uint8 OffsetHigh : 7;
 
 		/** If set indicates that there is another SRD following this one.*/
-		uint8 ContinuationBit		: 1;
-		uint8 OffsetLow				: 8;
+		uint8 ContinuationBit : 1;
+		uint8 OffsetLow : 8;
 
 
 		/** Returns SRD associated row number */
@@ -67,19 +88,10 @@ namespace UE::RivermaxCore::Private
 		void SetOffset(uint16 Offset);
 	};
 
-	/** Total packed size should be 26 octets. */
-	struct FRawRTPHeader
+	/** Total size should be 26 octets. */
+	struct FVideoRTPHeader
 	{
-		uint8 ContributingSourceCount   : 4;
-		uint8 ExtensionBit              : 1;
-		uint8 PaddingBit                : 1;
-		uint8 Version                   : 2;
-		uint8 PayloadType               : 7;
-		uint8 MarkerBit                 : 1;
-		uint16 SequenceNumber           : 16;
-		uint32 Timestamp                : 32;
-		uint32 SynchronizationSource    : 32;
-		uint16 ExtendedSequenceNumber   : 16;
+		FRawRTPHeader RTPHeader;
 
 		FRawSRD SRD1;                   // 48
 
@@ -89,6 +101,206 @@ namespace UE::RivermaxCore::Private
 		static constexpr uint32 OneSRDSize = 20;
 		static constexpr uint32 TwoSRDSize = 26;
 	};
+
+	/*
+	* https://datatracker.ietf.org/doc/html/rfc8331
+	*
+	*       0                   1                   2                   3
+	*       0 1 2 3 4 5 6 7 8 9 0 1 2 3 4 5 6 7 8 9 0 1 2 3 4 5 6 7 8 9 0 1
+	*       +-+-+-+-+-+-+-+-+-+-+-+-+-+-+-+-+-+-+-+-+-+-+-+-+-+-+-+-+-+-+-+-+
+	*       |           Length=32           | ANC_Count=2   | F | reserved  |
+	*       +-+-+-+-+-+-+-+-+-+-+-+-+-+-+-+-+-+-+-+-+-+-+-+-+-+-+-+-+-+-+-+-+
+	*       |    reserved                   |
+	*/
+	struct FAncRTPHeader 
+	{
+		FRawRTPHeader RTPHeader;
+
+		/**
+		* From RFC 8331:
+		*  Number of octets of the ANC data RTP payload, beginning with 
+		*  the "C" bit of the first ANC packet data header, as an
+        *  unsigned integer in network byte order.  Note that all
+        *  word_align fields contribute to the calculation of the Length
+        *  field.
+		*/
+		uint16 Length         : 16;
+		uint8 ANCCount        : 8;
+		uint8 ReservedHigh    : 6;
+
+		/**
+		*	These two bits relate to signaling the field specified by the
+        *   RTP timestamp in an interlaced SDI raster.  A value of 0b00
+        *   indicates that either the video format is progressive or that
+        *   no field is specified.
+		*/
+		uint8 Field           : 2;
+		uint16 ReservedLow    : 16;
+	};
+
+	/**
+	* For FAncillaryDataHeaderFields and FAncillaryTimecodeHeaderFields
+	*       0                   1                   2                   3
+	*       0 1 2 3 4 5 6 7 8 9 0 1 2 3 4 5 6 7 8 9 0 1 2 3 4 5 6 7 8 9 0 1
+	*       +-+-+-+-+-+-+-+-+-+-+-+-+-+-+-+-+-+-+-+-+-+-+-+-+-+-+-+-+-+-+-+-+
+	*       |C|   Line_Number=9     |   Horizontal_Offset   |S| StreamNum=0 |
+	*       +-+-+-+-+-+-+-+-+-+-+-+-+-+-+-+-+-+-+-+-+-+-+-+-+-+-+-+-+-+-+-+-+
+	*       |         DID       |        SDID       |  Data_Count=0x84  |
+	*       +-+-+-+-+-+-+-+-+-+-+-+-+-+-+-+-+-+-+-+-+-+-+-+-+-+-+-+-+-+-+-+-+
+	*                                User_Data_Words...
+	*       +-+-+-+-+-+-+-+-+-+-+-+-+-+-+-+-+-+-+-+-+-+-+-+-+-+-+-+-+-+-+-+-+
+	*                   |   Checksum_Word   |         word_align            |
+	*       +-+-+-+-+-+-+-+-+-+-+-+-+-+-+-+-+-+-+-+-+-+-+-+-+-+-+-+-+-+-+-+-+
+	*/
+	struct FAncillaryDataHeaderFields 
+	{
+		/**
+		*   +-------------+--------------------------------------------------------+
+		*	| Line_Number | ANC data packet generic vertical location              |
+		*	+-------------+--------------------------------------------------------+
+		*	|   0x7FF     | Without specific line location within the field or     |
+		*	|             | frame                                                  |
+		*/
+		uint8 LineNumberHigh : 7;
+
+		uint8 ColorDiff : 1;
+
+		/**
+		* ANC data packet generic horizontal location
+		* 0xFFF     | Without specific horizontal location   
+		*/
+		uint8 HorizontalOffsetHigh : 4;
+
+		uint8 LineNumberLow : 4;
+
+		uint8 HorizontalOffsetLow : 8;
+
+		uint8 StreamNum : 7;
+
+		/**
+		* Data Stream Flag.
+		* This field indicates whether the data stream number of a multi-stream data mapping used to
+		* transport the ANC data packet is specified.
+		*/
+		uint8 StreamFlag : 1;
+
+		inline void SetLineNumber(uint16 InLineNumber)
+		{
+			LineNumberLow = (0x000F & InLineNumber);
+			LineNumberHigh = (0x007F & (InLineNumber >> 4));
+		}
+
+		inline void SetHorizontalOffset(uint16 InHorizontalOffset)
+		{
+			HorizontalOffsetLow = (0x00FF & InHorizontalOffset);
+			HorizontalOffsetHigh = (0x000F & (InHorizontalOffset >> 8));
+		}
+		
+
+	};
+
+
+	/**
+	* This will contain ancillary timecode data packet starting from DID field.
+	*/
+	struct FAncillaryTimecodeHeaderFields
+	{
+		FAncillaryDataHeaderFields DataPacketHeaderFields;
+
+		// More on DID and SDID can be found here: https://smpte-ra.org/smpte-ancillary-data-smpte-st-291/
+		uint8 DIDHigh : 8;
+		uint8 SDIDHigh : 6;
+		uint8 DIDLow : 2;
+
+		/**
+		* From: https://datatracker.ietf.org/doc/html/rfc8331
+		* Data_Count: 10 bits The lower 8 bits of Data_Count, corresponding to bits b7
+		* (MSB; most significant bit) through b0 (LSB; least significant bit) of the 10-bit Data_Count word,
+		* contain the actual count of 10-bit words in User_Data_Words.  Bit b8 is the even parity for
+		* bits b7 through b0, and bit b9 is the inverse (logical NOT) of bit b8.
+		*/
+		uint8 DataCountHigh : 2;
+		uint8 DataCountEvenParity : 1;
+		uint8 DataCountInverseLogical : 1;
+		uint8 SDIDLow : 4;
+
+		uint8 DataCountLow : 6;
+		
+		uint32 UserData1 : 10;
+		uint32 UserData2 : 10;
+		uint32 UserData3 : 10;
+		uint32 UserData4 : 10;
+		uint32 UserData5 : 10;
+		uint32 UserData6 : 10;
+		uint32 UserData7 : 10;
+		uint32 UserData8 : 10;
+
+		uint8 ChecksumHigh : 2;
+
+		uint8 ChecksumLow : 8;
+
+		// Word align 1st to 6th bit 6 bits
+		uint8 WordAlign6thBit : 8;
+
+		/** Checks if the number of set bits is odd or even. Returns 0 if the number of 1s is even. 1 if odd. */
+		uint8 EvenParity(uint64 Num);
+
+		/** Add even parity check to the provided value. */
+		inline uint16 DidSdidAddParity(uint16 Value)
+		{
+			Value &= 0xff;
+			uint8 Parity = EvenParity(Value);
+
+			Value |= ((!Parity << 1) | Parity) << 8;
+			return Value;
+		}
+
+		/** Set DID value for this ANC header. */
+		inline void SetDID(uint16 InDID)
+		{
+			InDID = DidSdidAddParity(InDID);
+			DIDLow = (0x0003 & InDID);
+			DIDHigh = (0x00FF & (InDID >> 2));
+		}
+
+		/** Get DID value for this ANC header. */
+		inline uint16 GetDID()
+		{
+			return (((uint16)DIDLow)) | (((uint16)DIDHigh) << 2);
+		}
+
+		/** Set SDID value for this ANC header. */
+		inline void SetSDID(uint16 InSDID)
+		{
+			InSDID = DidSdidAddParity(InSDID);
+			SDIDLow = (0x000F & InSDID);
+			SDIDHigh = (0x003F & (InSDID >> 4));
+		}
+
+		/** Get SDID value for this ANC header. */
+		inline uint16 GetSDID()
+		{
+			return (((uint16)SDIDLow)) | (((uint16)SDIDHigh) << 4);
+		}
+
+		/** Set Data count. */
+		inline void SetDataCount(uint8 RawDataCount)
+		{
+			DataCountLow = (0x3F & (uint16)RawDataCount);
+			DataCountHigh = (0x03 & (RawDataCount >> 6));
+			DataCountEvenParity = EvenParity(GetDataCount());
+			DataCountInverseLogical = !DataCountEvenParity;
+		}
+
+		/** Get Data count. */
+		inline uint16 GetDataCount()
+		{
+			return (((uint16)DataCountLow)) | (((uint16)DataCountHigh) << 6);
+		}
+
+		/** Sets the ANC timecode. */
+		void SetATCTimecode(uint8 Hours, uint8 Minutes, uint8 Seconds, uint8 Frames, bool bDropFrame);
+	};
 #if PLATFORM_SUPPORTS_PRAGMA_PACK
 #pragma pack(pop)
 #endif
@@ -115,7 +327,7 @@ namespace UE::RivermaxCore::Private
 	struct FRTPHeader
 	{
 		FRTPHeader() = default;
-		FRTPHeader(const FRawRTPHeader& RawHeader);
+		FRTPHeader(const FVideoRTPHeader& VideoRTP);
 
 		/** Returns the total payload of this RTP */
 		uint16 GetTotalPayloadSize() const;
diff --git a/Engine/Plugins/VirtualProduction/Rivermax/RivermaxCore/Source/RivermaxCore/Private/RivermaxCoreModule.cpp b/Engine/Plugins/VirtualProduction/Rivermax/RivermaxCore/Source/RivermaxCore/Private/RivermaxCoreModule.cpp
--- a/Engine/Plugins/VirtualProduction/Rivermax/RivermaxCore/Source/RivermaxCore/Private/RivermaxCoreModule.cpp
+++ b/Engine/Plugins/VirtualProduction/Rivermax/RivermaxCore/Source/RivermaxCore/Private/RivermaxCoreModule.cpp
@@ -5,9 +5,10 @@
 #include "HAL/IConsoleManager.h"
 #include "Modules/ModuleManager.h"
 #include "RivermaxBoundaryMonitor.h"
-#include "RivermaxInputStream.h"
+#include "Streams/RivermaxInputStream.h"
 #include "RivermaxManager.h"
-#include "RivermaxOutputStream.h"
+#include "Streams/RivermaxOutVideoStream.h"
+#include "Streams/RivermaxOutAncStream.h"
 
 
 static TAutoConsoleVariable<int32> CVarRivermaxEnableBoundaryMonitor(
@@ -27,16 +28,29 @@ void FRivermaxCoreModule::ShutdownModule()
 	BoundaryMonitor->EnableMonitoring(false);
 }
 
-TUniquePtr<UE::RivermaxCore::IRivermaxInputStream> FRivermaxCoreModule::CreateInputStream()
+TUniquePtr<UE::RivermaxCore::IRivermaxInputStream> FRivermaxCoreModule::CreateInputStream(UE::RivermaxCore::ERivermaxStreamType, const TArray<char>& InSDPDescription)
 {
 	using UE::RivermaxCore::Private::FRivermaxInputStream;
 	return MakeUnique<FRivermaxInputStream>();
 }
 
-TUniquePtr<UE::RivermaxCore::IRivermaxOutputStream> FRivermaxCoreModule::CreateOutputStream()
+TUniquePtr<UE::RivermaxCore::IRivermaxOutputStream> FRivermaxCoreModule::CreateOutputStream(UE::RivermaxCore::ERivermaxStreamType StreamType, const TArray<char>& InSDPDescription)
 {
-	using UE::RivermaxCore::Private::FRivermaxOutputStream;
-	return MakeUnique<FRivermaxOutputStream>();
+	using UE::RivermaxCore::Private::FRivermaxOutVideoStream;
+	using UE::RivermaxCore::Private::FRivermaxOutAncStream;
+	switch (StreamType)
+	{
+	case UE::RivermaxCore::ERivermaxStreamType::VIDEO_2110_20_STREAM:
+		return MakeUnique<FRivermaxOutVideoStream>(InSDPDescription);
+		break;
+	case UE::RivermaxCore::ERivermaxStreamType::ANC_2110_40_STREAM:
+		return MakeUnique<FRivermaxOutAncStream>(InSDPDescription);
+		break;
+	case UE::RivermaxCore::ERivermaxStreamType::AUDIO_2110_30_STREAM:
+	default:
+		return nullptr;
+		break;
+	}
 }
 
 TSharedPtr<UE::RivermaxCore::IRivermaxManager> FRivermaxCoreModule::GetRivermaxManager()
diff --git a/Engine/Plugins/VirtualProduction/Rivermax/RivermaxCore/Source/RivermaxCore/Private/RivermaxCoreModule.h b/Engine/Plugins/VirtualProduction/Rivermax/RivermaxCore/Source/RivermaxCore/Private/RivermaxCoreModule.h
--- a/Engine/Plugins/VirtualProduction/Rivermax/RivermaxCore/Source/RivermaxCore/Private/RivermaxCoreModule.h
+++ b/Engine/Plugins/VirtualProduction/Rivermax/RivermaxCore/Source/RivermaxCore/Private/RivermaxCoreModule.h
@@ -4,7 +4,6 @@
 
 #include "IRivermaxCoreModule.h"
 
-
 namespace UE::RivermaxCore
 {
 	class IRivermaxManager;
@@ -23,8 +22,8 @@ public:
 	//~ End IModuleInterface
 
 	//~ Begin IRivermaxCoreModule interface
-	virtual TUniquePtr<UE::RivermaxCore::IRivermaxInputStream> CreateInputStream() override;
-	virtual TUniquePtr<UE::RivermaxCore::IRivermaxOutputStream> CreateOutputStream() override;
+	virtual TUniquePtr<UE::RivermaxCore::IRivermaxInputStream> CreateInputStream(UE::RivermaxCore::ERivermaxStreamType, const TArray<char>& InSDPDescription) override;
+	virtual TUniquePtr<UE::RivermaxCore::IRivermaxOutputStream> CreateOutputStream(UE::RivermaxCore::ERivermaxStreamType, const TArray<char>& InSDPDescription) override;
 	virtual TSharedPtr<UE::RivermaxCore::IRivermaxManager> GetRivermaxManager() override;
 	virtual UE::RivermaxCore::IRivermaxBoundaryMonitor& GetRivermaxBoundaryMonitor() override;
 	//~ End IRivermaxCoreModule interface
diff --git a/Engine/Plugins/VirtualProduction/Rivermax/RivermaxCore/Source/RivermaxCore/Private/RivermaxFrameManager.cpp b/Engine/Plugins/VirtualProduction/Rivermax/RivermaxCore/Source/RivermaxCore/Private/RivermaxFrameManager.cpp
--- a/Engine/Plugins/VirtualProduction/Rivermax/RivermaxCore/Source/RivermaxCore/Private/RivermaxFrameManager.cpp
+++ b/Engine/Plugins/VirtualProduction/Rivermax/RivermaxCore/Source/RivermaxCore/Private/RivermaxFrameManager.cpp
@@ -60,7 +60,7 @@ namespace UE::RivermaxCore::Private
 			{
 				/** Create actual frames and assign their video memory address from allocator */
 				TSharedPtr<FRivermaxOutputFrame> Frame = FramePool->AcquireShared(true /*allocate memory*/);
-				Frame->VideoBuffer = FrameAllocator->GetFrameAddress(Index);
+				Frame->Buffer = FrameAllocator->GetFrameAddress(Index);
 				AllocatedFrames.Add(Frame);
 			}
 			AllocatedFrames.Empty();
@@ -125,7 +125,7 @@ namespace UE::RivermaxCore::Private
 		OnFrameReadyDelegate.ExecuteIfBound();
 	}
 
-	bool FFrameManager::SetFrameData(const FRivermaxOutputVideoFrameInfo& NewFrameInfo, TSharedPtr<FRivermaxOutputFrame> ReservedFrame)
+	bool FFrameManager::SetFrameData(TSharedPtr<FRivermaxOutputVideoFrameInfo> NewFrameInfo, TSharedPtr<FRivermaxOutputFrame> ReservedFrame)
 	{
 		bool bSuccess = false;
 		check(ReservedFrame.IsValid());
@@ -134,13 +134,12 @@ namespace UE::RivermaxCore::Private
 		Sidecar->CopiedFrame = ReservedFrame;
 
 		FCopyArgs Args;
-		Args.RHISourceMemory = NewFrameInfo.GPUBuffer;
-		Args.SourceMemory = NewFrameInfo.CPUBuffer;
-		Args.DestinationMemory = ReservedFrame->VideoBuffer;
-		Args.SizeToCopy = NewFrameInfo.Height * NewFrameInfo.Stride;
+		Args.RHISourceMemory = NewFrameInfo->GPUBuffer;
+		Args.SourceMemory = NewFrameInfo->CPUBuffer;
+		Args.DestinationMemory = ReservedFrame->Buffer;
+		Args.SizeToCopy = NewFrameInfo->Height * NewFrameInfo->Stride;
 		Args.SideCar = MoveTemp(Sidecar);
 		bSuccess = FrameAllocator->CopyData(Args);
-		ReservedFrame->SetFrameCounter(NewFrameInfo.FrameIdentifier);
 
 		if (!bSuccess)
 		{
diff --git a/Engine/Plugins/VirtualProduction/Rivermax/RivermaxCore/Source/RivermaxCore/Private/RivermaxFrameManager.h b/Engine/Plugins/VirtualProduction/Rivermax/RivermaxCore/Source/RivermaxCore/Private/RivermaxFrameManager.h
--- a/Engine/Plugins/VirtualProduction/Rivermax/RivermaxCore/Source/RivermaxCore/Private/RivermaxFrameManager.h
+++ b/Engine/Plugins/VirtualProduction/Rivermax/RivermaxCore/Source/RivermaxCore/Private/RivermaxFrameManager.h
@@ -128,7 +128,7 @@ namespace UE::RivermaxCore::Private
 		void FrameSentEvent();
 
 		/** Initiates memory copy for a given frame */
-		bool SetFrameData(const FRivermaxOutputVideoFrameInfo& NewFrame, TSharedPtr<FRivermaxOutputFrame> ReservedFrame);
+		bool SetFrameData(TSharedPtr<FRivermaxOutputVideoFrameInfo> NewFrame, TSharedPtr<FRivermaxOutputFrame> ReservedFrame);
 
 	private:
 
diff --git a/Engine/Plugins/VirtualProduction/Rivermax/RivermaxCore/Source/RivermaxCore/Private/RivermaxInputStream.cpp b/Engine/Plugins/VirtualProduction/Rivermax/RivermaxCore/Source/RivermaxCore/Private/RivermaxInputStream.cpp
deleted file mode 100644
diff --git a/Engine/Plugins/VirtualProduction/Rivermax/RivermaxCore/Source/RivermaxCore/Private/RivermaxInputStream.h b/Engine/Plugins/VirtualProduction/Rivermax/RivermaxCore/Source/RivermaxCore/Private/RivermaxInputStream.h
deleted file mode 100644
diff --git a/Engine/Plugins/VirtualProduction/Rivermax/RivermaxCore/Source/RivermaxCore/Private/RivermaxOutputFrame.cpp b/Engine/Plugins/VirtualProduction/Rivermax/RivermaxCore/Source/RivermaxCore/Private/RivermaxOutputFrame.cpp
--- a/Engine/Plugins/VirtualProduction/Rivermax/RivermaxCore/Source/RivermaxCore/Private/RivermaxOutputFrame.cpp
+++ b/Engine/Plugins/VirtualProduction/Rivermax/RivermaxCore/Source/RivermaxCore/Private/RivermaxOutputFrame.cpp
@@ -5,7 +5,7 @@
 namespace UE::RivermaxCore::Private
 {
 	FRivermaxOutputFrame::FRivermaxOutputFrame()
-		: VideoBuffer(nullptr)
+		: Buffer(nullptr)
 	{
 		Clear();
 	}
diff --git a/Engine/Plugins/VirtualProduction/Rivermax/RivermaxCore/Source/RivermaxCore/Private/RivermaxOutputFrame.h b/Engine/Plugins/VirtualProduction/Rivermax/RivermaxCore/Source/RivermaxCore/Private/RivermaxOutputFrame.h
--- a/Engine/Plugins/VirtualProduction/Rivermax/RivermaxCore/Source/RivermaxCore/Private/RivermaxOutputFrame.h
+++ b/Engine/Plugins/VirtualProduction/Rivermax/RivermaxCore/Source/RivermaxCore/Private/RivermaxOutputFrame.h
@@ -27,10 +27,14 @@ namespace UE::RivermaxCore::Private
 		void Reset();
 
 		/** 
-		 * Video buffer where we copy texture to be sent out 
+		 * Buffer where we copy our data.
+		 * 
+		 * In Case of Video stream this is where we copy texture to be sent out 
 		 * If using GPUDirect, memory will be allocated in cuda space
+		 * 
+		 * In case of ANCILLARY this will contain our ancillary 
 		 */
-		void* VideoBuffer = nullptr;
+		void* Buffer = nullptr;
 
 		/** Timecode at which this frame was captured */
 		FTimecode Timecode;
diff --git a/Engine/Plugins/VirtualProduction/Rivermax/RivermaxCore/Source/RivermaxCore/Private/RivermaxOutputStream.cpp b/Engine/Plugins/VirtualProduction/Rivermax/RivermaxCore/Source/RivermaxCore/Private/RivermaxOutputStream.cpp
deleted file mode 100644
diff --git a/Engine/Plugins/VirtualProduction/Rivermax/RivermaxCore/Source/RivermaxCore/Private/RivermaxOutputStream.h b/Engine/Plugins/VirtualProduction/Rivermax/RivermaxCore/Source/RivermaxCore/Private/RivermaxOutputStream.h
deleted file mode 100644
diff --git a/Engine/Plugins/VirtualProduction/Rivermax/RivermaxCore/Source/RivermaxCore/Private/RivermaxUtils.cpp b/Engine/Plugins/VirtualProduction/Rivermax/RivermaxCore/Source/RivermaxCore/Private/RivermaxUtils.cpp
--- a/Engine/Plugins/VirtualProduction/Rivermax/RivermaxCore/Source/RivermaxCore/Private/RivermaxUtils.cpp
+++ b/Engine/Plugins/VirtualProduction/Rivermax/RivermaxCore/Source/RivermaxCore/Private/RivermaxUtils.cpp
@@ -9,136 +9,47 @@
 
 namespace UE::RivermaxCore::Private::Utils
 {
-	FString PixelFormatToSamplingDesc(ESamplingType SamplingType)
+	uint32 TimestampToFrameNumber(uint32 Timestamp, const FFrameRate& FrameRate)
 	{
-		switch (SamplingType)
-		{
-			case ESamplingType::RGB_8bit:
-			case ESamplingType::RGB_10bit:
-			case ESamplingType::RGB_12bit:
-			case ESamplingType::RGB_16bit:
-			case ESamplingType::RGB_16bitFloat:
-			{
-				return FString(TEXT("RGB"));
-			}
-			case ESamplingType::YUV422_8bit:
-			case ESamplingType::YUV422_10bit:
-			case ESamplingType::YUV422_12bit:
-			case ESamplingType::YUV422_16bit:
-			case ESamplingType::YUV422_16bitFloat:
-			{
-				return FString(TEXT("YCbCr-4:2:2"));
-			}
-			case ESamplingType::YUV444_8bit:
-			case ESamplingType::YUV444_10bit:
-			case ESamplingType::YUV444_12bit:
-			case ESamplingType::YUV444_16bit:
-			case ESamplingType::YUV444_16bitFloat:
-			{
-				return FString(TEXT("YCbCr-4:4:4"));
-			}
-			default:
-			{
-				checkNoEntry();
-				return FString();
-			}
-		}
+		using namespace UE::RivermaxCore::Private::Utils;
+		const double MediaFrameTime = Timestamp / MediaClockSampleRate;
+		const uint32 FrameNumber = FMath::Floor(MediaFrameTime * FrameRate.AsDecimal());
+		return FrameNumber;
 	}
 
-	FString PixelFormatToBitDepth(ESamplingType SamplingType)
+	/** Returns a mediaclock timestamp, for rtp, based on a clock time */
+	uint32 GetTimestampFromTime(uint64 InTimeNanosec, double InMediaClockRate)
 	{
-		switch (SamplingType)
-		{
-		case ESamplingType::RGB_8bit:
-		case ESamplingType::YUV422_8bit:
-		case ESamplingType::YUV444_8bit:
-		{
-			return FString(TEXT("8"));
-		}
-		case ESamplingType::RGB_10bit:
-		case ESamplingType::YUV422_10bit:
-		case ESamplingType::YUV444_10bit:
-		{
-			return FString(TEXT("10"));
-		}
-		case ESamplingType::RGB_12bit:
-		case ESamplingType::YUV422_12bit:
-		case ESamplingType::YUV444_12bit:
-		{
-			return FString(TEXT("12"));
-		}
-		case ESamplingType::RGB_16bit:
-		case ESamplingType::YUV422_16bit:
-		case ESamplingType::YUV444_16bit:
-		{
-			return FString(TEXT("16"));
-		}
-		case ESamplingType::RGB_16bitFloat:
-		case ESamplingType::YUV422_16bitFloat:
-		case ESamplingType::YUV444_16bitFloat:
-		{
-			return FString(TEXT("16f"));
-		}
-		default:
-		{
-			checkNoEntry();
-			return FString();
-		}
-		}
+		// RTP timestamp is 32 bits and based on media clock (usually 90kHz).
+		// Conversion based on rivermax samples
+
+		const uint64 Nanoscale = 1E9;
+		const uint64 Seconds = InTimeNanosec / Nanoscale;
+		const uint64 Nanoseconds = InTimeNanosec % Nanoscale;
+		const uint64 MediaFrameNumber = Seconds * InMediaClockRate;
+		const uint64 MediaSubFrameNumber = Nanoseconds * InMediaClockRate / Nanoscale;
+		const double Mask = 0x100000000;
+		const double MediaTime = FMath::Fmod(MediaFrameNumber, Mask);
+		const double MediaSubTime = FMath::Fmod(MediaSubFrameNumber, Mask);
+		return MediaTime + MediaSubTime;
 	}
 
-	void StreamOptionsToSDPDescription(const UE::RivermaxCore::FRivermaxOutputStreamOptions& Options, FAnsiStringBuilderBase& OutSDPDescription)
+	FTimecode GetTimecodeFromTime(uint64 InTimeNanosec, double InMediaClockRate, const FFrameRate& FrameRate)
 	{
-		// Basic SDP string creation from a set of options. At some point, having a proper SDP loader / creator would be useful.
-		// Refer to https://datatracker.ietf.org/doc/html/rfc4570
+		FTimespan Timespan = FTimespan(InTimeNanosec / ETimespan::NanosecondsPerTick);
 
-		// Apply desired multiplier and convert fractional frame rate to be represented over 1001 for sdp compliance.
-		FFrameRate DesiredRate = Options.FrameRate;
-		const double DecimalDesiredRate = Options.FrameRate.AsDecimal();
-		const bool bIsNonStandardFractionalFrameRate = (DesiredRate.Denominator != 1001 && !FMath::IsNearlyZero(FMath::Frac(DecimalDesiredRate)));
-		if (bIsNonStandardFractionalFrameRate)
-		{
-			UE_LOG(LogRivermax, Warning, TEXT("Fractional frame rates must be described using a denominator of 1001. Converting it for stream creation."));
-		}
+		// This should come from timecode provider or media profile.
+		constexpr int32 DaylightSavingTimeHourOffset = 0; 
+		constexpr int32 UTCSecondsOffset = 37;
 
-		FString FrameRateDescription;
-		if (bIsNonStandardFractionalFrameRate)
-		{
-			const double NewRateDecimal = DecimalDesiredRate * 1001;
-			DesiredRate.Numerator = FMath::RoundToInt32(NewRateDecimal);
-			DesiredRate.Denominator = 1001;
-		}
-		
-		if(DesiredRate.Denominator == 1001)
-		{
-			FrameRateDescription = FString::Printf(TEXT("%u/%u"), DesiredRate.Numerator, DesiredRate.Denominator);
-		}
-		else
-		{
-			FrameRateDescription = FString::Printf(TEXT("%u"), FMath::RoundToInt32(DesiredRate.AsDecimal()));
-		}
+		// Adjust for daylight saving that might be required. 
+		Timespan -= FTimespan(DaylightSavingTimeHourOffset, 0, 0);
 
-		constexpr int32 MulticastTTL = 64;
-		OutSDPDescription.Appendf("v=0\n");
-		OutSDPDescription.Appendf("s=SMPTE ST2110 20 streams\n");
-		OutSDPDescription.Appendf("t=0 0\n");
-		OutSDPDescription.Appendf("m=video %d RTP/AVP 96\n", Options.Port);
-		OutSDPDescription.Appendf("c=IN IP4 %S/%d\n", *Options.StreamAddress, MulticastTTL);
-		OutSDPDescription.Appendf("a=source-filter: incl IN IP4 %S %S\n", *Options.StreamAddress, *Options.InterfaceAddress);
-		OutSDPDescription.Appendf("a=rtpmap:96 raw/90000\n");
-		OutSDPDescription.Appendf("a=fmtp: 96 sampling=%S; width=%d; height=%d; exactframerate=%S; depth=%S; TCS=SDR; colorimetry=BT709; PM=2110GPM; SSN=ST2110-20:2017; TP=2110TPN;\n"
-			, *PixelFormatToSamplingDesc(Options.PixelFormat)
-			, Options.AlignedResolution.X
-			, Options.AlignedResolution.Y
-			, *FrameRateDescription
-			, *PixelFormatToBitDepth(Options.PixelFormat));
-	}
+		// Convert from TAI PTP Time to UTC
+		Timespan -= FTimespan(0, 0, UTCSecondsOffset);
 
-	uint32 TimestampToFrameNumber(uint32 Timestamp, const FFrameRate& FrameRate)
-	{
-		using namespace UE::RivermaxCore::Private::Utils;
-		const double MediaFrameTime = Timestamp / MediaClockSampleRate;
-		const uint32 FrameNumber = FMath::Floor(MediaFrameTime * FrameRate.AsDecimal());
-		return FrameNumber;
+		constexpr bool bRollOver = true;
+
+		return FTimecode::FromTimespan(Timespan, FrameRate, bRollOver);
 	}
 }
\ No newline at end of file
diff --git a/Engine/Plugins/VirtualProduction/Rivermax/RivermaxCore/Source/RivermaxCore/Private/RivermaxUtils.h b/Engine/Plugins/VirtualProduction/Rivermax/RivermaxCore/Source/RivermaxCore/Private/RivermaxUtils.h
--- a/Engine/Plugins/VirtualProduction/Rivermax/RivermaxCore/Source/RivermaxCore/Private/RivermaxUtils.h
+++ b/Engine/Plugins/VirtualProduction/Rivermax/RivermaxCore/Source/RivermaxCore/Private/RivermaxUtils.h
@@ -5,10 +5,11 @@
 #include "CoreMinimal.h"
 
 #include "RivermaxTypes.h"
+#include "Misc/Timecode.h"
 
 namespace UE::RivermaxCore
 {
-	struct FRivermaxOutputStreamOptions;
+	struct FRivermaxOutputOptions;
 }
 
 namespace UE::RivermaxCore::Private::Utils
@@ -30,12 +31,17 @@ namespace UE::RivermaxCore::Private::Utils
 	/** Common sleep time used in places where we are waiting for something to complete */
 	static constexpr float SleepTimeSeconds = 50 * 1E-6;
 
-	/** Convert a set of streaming option to its SDP description. Currently only support video type. */
-	void StreamOptionsToSDPDescription(const UE::RivermaxCore::FRivermaxOutputStreamOptions& Options, FAnsiStringBuilderBase& OutSDPDescription);
-
 	/**
 	 * Converts a timestamp in MediaClock period units to a frame number for a given frame rate
 	 * 2110-20 streams uses a standard media clock rate of 90kHz
 	 */
 	uint32 TimestampToFrameNumber(uint32 Timestamp, const FFrameRate& FrameRate);
+
+	/** Returns a mediaclock timestamp, for rtp, based on a clock time */
+	uint32 GetTimestampFromTime(uint64 InTimeNanosec, double InMediaClockRate);
+
+	/** Convert PTP time to timecode. */
+	FTimecode GetTimecodeFromTime(uint64 InTimeNanosec, double InMediaClockRate, const FFrameRate& FrameRate);
+	
+
 }
diff --git a/Engine/Plugins/VirtualProduction/Rivermax/RivermaxCore/Source/RivermaxCore/Private/Streams/RivermaxInputStream.cpp b/Engine/Plugins/VirtualProduction/Rivermax/RivermaxCore/Source/RivermaxCore/Private/Streams/RivermaxInputStream.cpp
new file mode 100644
--- /dev/null
+++ b/Engine/Plugins/VirtualProduction/Rivermax/RivermaxCore/Source/RivermaxCore/Private/Streams/RivermaxInputStream.cpp
@@ -0,0 +1,1263 @@
+// Copyright Epic Games, Inc. All Rights Reserved.
+
+#include "RivermaxInputStream.h"
+
+#include "Async/Async.h"
+#include "CudaModule.h"
+#include "HAL/CriticalSection.h"
+#include "ID3D12DynamicRHI.h"
+#include "IRivermaxCoreModule.h"
+#include "IRivermaxManager.h"
+#include "Misc/ByteSwap.h"
+#include "RenderGraphUtils.h"
+#include "RivermaxLog.h"
+#include "RivermaxPTPUtils.h"
+#include "RivermaxTracingUtils.h"
+#include "RivermaxUtils.h"
+#include "RTPHeader.h"
+
+#if PLATFORM_WINDOWS
+#include <WS2tcpip.h>
+#endif
+
+namespace UE::RivermaxCore::Private
+{
+	/** Initialization calls to Rivermax SDK isn't threadsafe and to prevent stalling game thread, a critical section is shared across streams */
+	static FCriticalSection InitCriticalSection;
+
+	static TAutoConsoleVariable<float> CVarWaitForCompletionTimeout(
+		TEXT("Rivermax.Input.WaitCompletionTimeout"),
+		0.25,
+		TEXT("Maximum time to wait, in seconds, when waiting for a memory copy operation to complete on the gpu."),
+		ECVF_Default);
+
+	static TAutoConsoleVariable<int32> CVarExpectedPayloadSize(
+		TEXT("Rivermax.Input.ExpectedPayloadSize"),
+		1500,
+		TEXT("Expected payload size used to initialize rivermax stream."),
+		ECVF_Default);
+
+	static bool GEnableLargeFirstPacketIntervalLogging = false;
+	FAutoConsoleVariableRef CVarRivermaxEnableLargeFirstPacketIntervalLogging(
+		TEXT("Rivermax.Input.EnableLargeFistPacketIntervalLogging")
+		, UE::RivermaxCore::Private::GEnableLargeFirstPacketIntervalLogging
+		, TEXT("Enables detection and logging of large delta times between frame boundary and first packet reception.")
+		, ECVF_Default);
+
+	static bool GClearFirstPacketIntervalStats = false;
+	FAutoConsoleVariableRef CVarRivermaxClearFirstPacketIntervalStats(
+		TEXT("Rivermax.Input.ClearFistPacketIntervalStats")
+		, UE::RivermaxCore::Private::GClearFirstPacketIntervalStats
+		, TEXT("Clears stats related to first packet interval detection.")
+		, ECVF_Default);
+
+	static int32 GLargeFirstPacketIntervalThresholdMicroSec = 2000;
+	FAutoConsoleVariableRef CVarRivermaxLargeFirstPacketIntervalThreshold(
+		TEXT("Rivermax.Input.LargeFistPacketIntervalThreshold")
+		, UE::RivermaxCore::Private::GLargeFirstPacketIntervalThresholdMicroSec
+		, TEXT("Microseconds required to consider a first packet interval to be large and be logged.")
+		, ECVF_Default);
+
+
+	void FFrameDescriptionTrackingData::ResetSingleFrameTracking()
+	{
+		PayloadSizeReceived.Empty();
+	}
+
+	void FFrameDescriptionTrackingData::EvaluateNewRTP(const FRTPHeader& NewHeader)
+	{
+		UpdateResolutionDetection(NewHeader);
+		UpdatePayloadSizeTracking(NewHeader);
+
+		if (NewHeader.bIsMarkerBit)
+		{
+			ResetSingleFrameTracking();
+		}
+	}
+
+	void FFrameDescriptionTrackingData::UpdateResolutionDetection(const FRTPHeader& NewHeader)
+	{
+		if (NewHeader.bIsMarkerBit)
+		{
+			const FVideoFormatInfo ExpectedFormatInfo = FStandardVideoFormat::GetVideoFormatInfo(ExpectedSamplingType);
+			const uint16 LastPayloadSize = NewHeader.GetLastPayloadSize();
+			if (LastPayloadSize % ExpectedFormatInfo.PixelGroupSize == 0)
+			{
+				DetectedResolution.X = NewHeader.GetLastRowOffset() + (LastPayloadSize / ExpectedFormatInfo.PixelGroupSize * ExpectedFormatInfo.PixelGroupCoverage);
+				DetectedResolution.Y = NewHeader.GetLastRowNumber() + 1;
+				
+				// Renable logging when a valid incoming frame was detected
+				bHasLoggedSamplingWarning = false;
+			}
+			else if(!bHasLoggedSamplingWarning)
+			{
+				bHasLoggedSamplingWarning = true;
+				UE_LOG(LogRivermax, Warning, TEXT("Detected incoming signal with unexpected sampling type."));
+			}
+		}
+	}
+
+	void FFrameDescriptionTrackingData::UpdatePayloadSizeTracking(const FRTPHeader& NewHeader)
+	{
+		if (NewHeader.SRD1.Length > 0)
+		{
+			uint32 PayloadSize = NewHeader.SRD1.Length;
+			if (NewHeader.SRD1.bHasContinuation)
+			{
+				PayloadSize += NewHeader.SRD2.Length;
+			}
+
+			if (PayloadSizeReceived.IsEmpty())
+			{
+				CommonPayloadSize = PayloadSize;
+			}
+			else if (!NewHeader.bIsMarkerBit && PayloadSize != CommonPayloadSize)
+			{
+				CommonPayloadSize = -1;
+			}
+
+			PayloadSizeReceived.FindOrAdd(PayloadSize) += 1;
+		}
+	}
+
+	bool FFrameDescriptionTrackingData::HasDetectedValidResolution() const
+	{
+		return DetectedResolution.X > 0 && DetectedResolution.Y > 0;
+	}
+
+	FRivermaxInputStream::FRivermaxInputStream()
+	{
+
+	}
+
+	FRivermaxInputStream::~FRivermaxInputStream()
+	{
+		Uninitialize();
+	}
+
+	bool FRivermaxInputStream::Initialize(const FRivermaxInputStreamOptions& InOptions, IRivermaxInputStreamListener& InListener)
+	{
+		IRivermaxCoreModule& RivermaxModule = FModuleManager::LoadModuleChecked<IRivermaxCoreModule>(TEXT("RivermaxCore"));
+		if (RivermaxModule.GetRivermaxManager()->ValidateLibraryIsLoaded() == false)
+		{
+			UE_LOG(LogRivermax, Warning, TEXT("Can't create Rivermax Input Stream. Library isn't initialized."));
+			return false;
+		}
+
+		Options = InOptions;
+		Listener = &InListener;
+		FormatInfo = FStandardVideoFormat::GetVideoFormatInfo(Options.PixelFormat);
+		ExpectedPayloadSize = CVarExpectedPayloadSize.GetValueOnGameThread();
+		bIsDynamicHeaderEnabled = RivermaxModule.GetRivermaxManager()->EnableDynamicHeaderSupport(Options.InterfaceAddress);
+		CachedAPI = RivermaxModule.GetRivermaxManager()->GetApi();
+		checkSlow(CachedAPI);
+
+		if (Options.bEnforceVideoFormat)
+		{
+			StreamResolution = Options.EnforcedResolution;
+		}
+		
+		// For now, we don't try to detect pixel format so use what the user selected
+		FrameDescriptionTracking.ExpectedSamplingType = Options.PixelFormat;
+
+		InitTaskFuture = Async(EAsyncExecution::TaskGraph, [this]()
+		{
+			TRACE_CPUPROFILER_EVENT_SCOPE(FRivermaxInputStream::InitStream);
+			FScopeLock Lock(&InitCriticalSection);
+
+			// If the stream is trying to shutdown before the init task has even started, don't bother
+			if (bIsShuttingDown)
+			{
+				return;
+			}
+
+			// Initialize stream builder
+			CachedAPI->rmx_input_init_stream(&StreamParameters, RMX_INPUT_RAW_PACKET);
+
+			FString ErrorMessage;
+			FNetworkSettings NetworkSettings;
+			bool bWasSuccessful = InitializeNetworkSettings(NetworkSettings, ErrorMessage);
+			if(bWasSuccessful)
+			{
+				bWasSuccessful = InitializeStreamParameters(ErrorMessage);
+
+				if (bWasSuccessful)
+				{
+					bWasSuccessful = FinalizeStreamCreation(NetworkSettings, ErrorMessage);
+				}
+			}
+
+			if (bWasSuccessful)
+			{
+				// Stream creation succeeded, launch reception thread
+				bIsActive = true;
+				RivermaxThread = FRunnableThread::Create(this, TEXT("Rivermax InputStream Thread"), 128 * 1024, TPri_AboveNormal, FPlatformAffinity::GetPoolThreadMask());
+
+				UE_LOG(LogRivermax, Display, TEXT("Input stream started listening to stream %s:%d on interface %s%s")
+					, *Options.StreamAddress
+					, Options.Port
+					, *Options.InterfaceAddress
+					, bIsUsingGPUDirect ? TEXT(" using GPUDirect") : TEXT(""));
+			}
+			else
+			{
+				UE_LOG(LogRivermax, Warning, TEXT("%s"), *ErrorMessage);
+			}
+			
+			FRivermaxInputInitializationResult Result;
+			Result.bHasSucceed = bWasSuccessful;
+			Result.bIsGPUDirectSupported = bIsUsingGPUDirect;
+			Listener->OnInitializationCompleted(Result);
+		});
+		
+		return true;
+	}
+
+	void FRivermaxInputStream::Uninitialize()
+	{
+		bIsShuttingDown = true;
+
+		//If init task is ongoing, wait till it's done
+		if (InitTaskFuture.IsReady() == false)
+		{
+			InitTaskFuture.Wait();
+		}
+
+		if (RivermaxThread != nullptr)
+		{
+			Stop();
+			RivermaxThread->Kill(true);
+			delete RivermaxThread;
+			RivermaxThread = nullptr;
+			UE_LOG(LogRivermax, Log, TEXT("Rivermax Input stream has shutdown"));
+		}
+
+		if (bIsDynamicHeaderEnabled)
+		{
+			if (IRivermaxCoreModule* RivermaxModule = FModuleManager::GetModulePtr<IRivermaxCoreModule>(TEXT("RivermaxCore")))
+			{
+				RivermaxModule->GetRivermaxManager()->DisableDynamicHeaderSupport(Options.InterfaceAddress);
+			}
+			bIsDynamicHeaderEnabled = false;
+		}
+
+		DeallocateBuffers();
+	}
+
+	void FRivermaxInputStream::Process_AnyThread()
+	{
+		const size_t MinChunkSize = 0;
+		const size_t MaxChunkSize = 5000;
+		const int Timeout = 0;
+		const int Flags = 0;
+		const rmx_input_completion* Completion = nullptr;
+		
+		checkSlow(CachedAPI);
+		rmx_status Status = CachedAPI->rmx_input_get_next_chunk(&BufferConfiguration.ChunkHandle);
+		if (Status == RMX_OK)
+		{
+			Completion = CachedAPI->rmx_input_get_chunk_completion(&BufferConfiguration.ChunkHandle);
+			if (Completion)
+			{
+				ParseChunks(Completion);
+			}
+		}
+		else
+		{
+			UE_LOG(LogRivermax, Warning, TEXT("Rivermax Input stream failed to get next chunk. Status: %d"), Status);
+		}
+
+		FPlatformProcess::SleepNoStats(UE::RivermaxCore::Private::Utils::SleepTimeSeconds);
+	}
+
+	bool FRivermaxInputStream::Init()
+	{
+		return true;
+	}
+
+	uint32 FRivermaxInputStream::Run()
+	{
+		while (bIsActive)
+		{
+			Process_AnyThread();
+			LogStats();
+		}
+
+		if (StreamId)
+		{
+			checkSlow(CachedAPI);
+			rmx_status Status = CachedAPI->rmx_input_detach_flow(StreamId, &FlowAttribute);
+			if (Status != RMX_OK)
+			{
+				UE_LOG(LogRivermax, Warning, TEXT("Failed to detach rivermax flow %d from input stream %d. Status: %d"), FlowTag, StreamId, Status);
+			}
+
+			Status = CachedAPI->rmx_input_destroy_stream(StreamId);
+
+			if (Status != RMX_OK)
+			{
+				UE_LOG(LogRivermax, Warning, TEXT("Failed to destroy input stream %d correctly. Status: %d"), StreamId, Status);
+			}
+		}
+
+		return 0;
+	}
+
+	void FRivermaxInputStream::Stop()
+	{
+		bIsActive = false;
+	}
+
+	void FRivermaxInputStream::Exit()
+	{
+
+	}
+
+	void FRivermaxInputStream::ParseChunks(const rmx_input_completion* Completion)
+	{
+		const size_t PacketCount = rmx_input_get_completion_chunk_size(Completion);
+
+		// We need to get all samples once per group of packets and store against its frame number.
+		for (uint64 StrideIndex = 0; StrideIndex < PacketCount; ++StrideIndex)
+		{
+			++StreamStats.ChunksReceived;
+
+			const uint8* RawHeaderPtr = reinterpret_cast<const uint8_t*>(rmx_input_get_completion_ptr(Completion, BufferConfiguration.HeaderBlockID));
+			const uint8* DataPtr = reinterpret_cast<const uint8_t*>(rmx_input_get_completion_ptr(Completion, BufferConfiguration.DataBlockID));
+
+			RawHeaderPtr += StrideIndex * BufferConfiguration.HeaderStride;
+			DataPtr += StrideIndex * BufferConfiguration.DataStride;
+
+			const rmx_input_packet_info* PacketInfo = CachedAPI->rmx_input_get_packet_info(&BufferConfiguration.ChunkHandle, StrideIndex);
+			const size_t PacketSize = rmx_input_get_packet_size(PacketInfo, BufferConfiguration.DataBlockID);
+			const size_t HeaderSize = rmx_input_get_packet_size(PacketInfo, BufferConfiguration.HeaderBlockID);
+			if (PacketSize > 0 && HeaderSize > 0)
+			{
+				if (FlowTag)
+				{
+					const uint32 PacketTag = rmx_input_get_packet_flow_tag(PacketInfo);
+					if (PacketTag != FlowTag)
+					{
+						UE_LOG(LogRivermax, Error, TEXT("Received data from unexpected FlowTag '%d'. Expected '%d'."), PacketTag, FlowTag);
+						Listener->OnStreamError();
+						bIsShuttingDown = true;
+						return;
+					}
+				}
+
+				const uint64 PacketTimstamp = rmx_input_get_packet_timestamp(PacketInfo);
+				
+
+				// Get RTPHeader address from the raw net header
+				const FVideoRTPHeader& RawRTPHeaderPtr = reinterpret_cast<const FVideoRTPHeader&>(*GetRTPHeaderPointer(RawHeaderPtr));
+				if (RawRTPHeaderPtr.RTPHeader.Version == 2)
+				{
+					FRTPHeader RTPHeader(RawRTPHeaderPtr);
+					// Add trace for the first packet of a frame to help visualize reception of a full frame in time
+					if (bIsFirstPacketReceived == false)
+					{
+						const uint32 FrameNumber = Utils::TimestampToFrameNumber(RTPHeader.Timestamp, Options.FrameRate);
+						TRACE_CPUPROFILER_EVENT_SCOPE_TEXT(*FRivermaxTracingUtils::RmaxInStartingFrameTraceEvents[FrameNumber % 10]);
+
+						IRivermaxCoreModule& RivermaxModule = FModuleManager::LoadModuleChecked<IRivermaxCoreModule>(TEXT("RivermaxCore"));
+						const uint64 CurrentTime = RivermaxModule.GetRivermaxManager()->GetTime();
+
+						bIsFirstPacketReceived = true;
+
+						if (StreamStats.EndOfFrameReceived > 0)
+						{
+							const uint64 FrameBoundary = GetAlignmentPointFromFrameNumber(GetFrameNumber(CurrentTime, Options.FrameRate), Options.FrameRate);
+							const uint64 FirstPacketInterval = CurrentTime - FrameBoundary;
+							StreamStats.MinFirstPacketIntervalNS = FMath::Min(StreamStats.MinFirstPacketIntervalNS, FirstPacketInterval);
+							StreamStats.MaxFirstPacketIntervalNS = FMath::Max(StreamStats.MaxFirstPacketIntervalNS, FirstPacketInterval);
+							
+							constexpr double Alpha = 0.8;
+							StreamStats.FirstPacketIntervalAccumulatorNS = Alpha * FirstPacketInterval + ((1.0 - Alpha) * StreamStats.FirstPacketIntervalAccumulatorNS);
+
+							if (GClearFirstPacketIntervalStats)
+							{
+								GClearFirstPacketIntervalStats = false;
+								StreamStats.MinFirstPacketIntervalNS = TNumericLimits<uint64>::Max();
+								StreamStats.MaxFirstPacketIntervalNS = TNumericLimits<uint64>::Min();
+								StreamStats.FirstPacketIntervalAccumulatorNS = 0;
+							}
+
+							if (GEnableLargeFirstPacketIntervalLogging)
+							{
+								const uint32 IntervalMicroSec = FirstPacketInterval / 1000;
+								if (IntervalMicroSec > (uint32)GLargeFirstPacketIntervalThresholdMicroSec)
+								{
+
+									UE_LOG(LogRivermax, Warning, TEXT("Large First packet interval: %llu. Min: %llu. Max: %llu. Avg: %llu.")
+									, FirstPacketInterval
+									, StreamStats.MinFirstPacketIntervalNS
+									, StreamStats.MaxFirstPacketIntervalNS
+									, StreamStats.FirstPacketIntervalAccumulatorNS);
+								}
+							}
+						}
+					}
+
+					StreamStats.BytesReceived += PacketSize + HeaderSize;
+					
+					UpdateFrameTracking(RTPHeader);
+
+					// Get video sample to write into (once per frame)
+					if (!StreamData.CurrentSamples.Contains(IRivermaxSample::ESampleType::Video) && State != EReceptionState::FrameError)
+					{
+						TSharedPtr<IRivermaxVideoSample> Sample = GetVideoSampleForReception(RTPHeader);
+						StreamData.CurrentSamples.Add(IRivermaxSample::ESampleType::Video, StaticCastSharedPtr<IRivermaxSample>(Sample));
+					}
+
+					switch (State)
+					{
+					case EReceptionState::Receiving:
+					{
+						FrameReceptionState(RTPHeader, DataPtr, StreamData.CurrentSamples[IRivermaxSample::ESampleType::Video]);
+						break;
+					}
+					case EReceptionState::WaitingForMarker:
+					{
+						WaitForMarkerState(RTPHeader);
+						break;
+					}
+					case EReceptionState::FrameError:
+					{
+						FrameErrorState(RTPHeader);
+						break;
+					}
+					default:
+					{
+						checkNoEntry();
+					}
+					}
+				}
+				else
+				{
+					++StreamStats.InvalidHeadercount;
+				}
+			}
+			else
+			{
+				++StreamStats.EmptyCompletionCount;
+			}
+		}
+	}
+
+	TSharedPtr<IRivermaxVideoSample> FRivermaxInputStream::GetVideoSampleForReception(const FRTPHeader& RTPHeader)
+	{
+		using namespace UE::RivermaxCore::Private::Utils;
+
+		FRivermaxInputVideoFrameDescriptor Descriptor;
+		Descriptor.Timestamp = RTPHeader.Timestamp;
+		Descriptor.FrameNumber = TimestampToFrameNumber(RTPHeader.Timestamp, Options.FrameRate);
+		Descriptor.Width = StreamResolution.X;
+		Descriptor.Height = StreamResolution.Y;
+		Descriptor.PixelFormat = Options.PixelFormat;
+		const uint32 PixelCount = StreamResolution.X * StreamResolution.Y;
+		const uint32 FrameSize = PixelCount / FormatInfo.PixelGroupCoverage * FormatInfo.PixelGroupSize;
+		Descriptor.VideoBufferSize = FrameSize;
+		FRivermaxInputVideoFrameRequest Request;
+		TSharedPtr<IRivermaxVideoSample> OutSample = Listener->OnVideoFrameRequested(Descriptor);
+
+		// Reset current frame to know when we have a valid one
+		StreamData.CurrentFrameVideoBuffer = nullptr;
+		if (OutSample.IsValid())
+		{
+			if (bIsUsingGPUDirect)
+			{
+				FBufferRHIRef RHIBuffer = OutSample->GetGPUBuffer()->GetRHI();
+				if (RHIBuffer.IsValid())
+				{
+					StreamData.CurrentFrameVideoBuffer = GetMappedBuffer(RHIBuffer);
+				}
+			}
+			else
+			{
+				StreamData.CurrentFrameVideoBuffer = OutSample->GetVideoBufferRawPtr(Descriptor.VideoBufferSize);
+			}
+		}
+
+		// Verify if we were able to request a valid frame. If engine is blocked, it could happen that there is none available 
+		if (StreamData.CurrentFrameVideoBuffer == nullptr)
+		{
+			// If we failed getting one, reset the valid first frame received and wait for the next one
+			UE_LOG(LogRivermax, Verbose, TEXT("Could not get a new frame for incoming frame with timestamp %u and frame number %u"), RTPHeader.Timestamp, Descriptor.FrameNumber);
+			Listener->OnVideoFrameReceptionError(OutSample);
+			State = EReceptionState::FrameError;
+			FrameErrorState(RTPHeader);
+		}
+		else
+		{
+			StreamData.WritingOffset = 0;
+			StreamData.ReceivedSize = 0;
+			StreamData.ExpectedSize = Descriptor.VideoBufferSize;
+			StreamData.DeviceWritePointerOne = nullptr;
+			StreamData.SizeToWriteOne = 0;
+			StreamData.DeviceWritePointerTwo = nullptr;
+			StreamData.SizeToWriteTwo = 0;
+			bIsFirstPacketReceived = false;
+		}
+
+		return OutSample;
+	}
+
+	void FRivermaxInputStream::LogStats()
+	{
+		static constexpr double LoggingInterval = 1.0;
+
+		const double CurrentTime = FPlatformTime::Seconds();
+		if (CurrentTime - LastLoggingTimestamp >= LoggingInterval)
+		{
+			LastLoggingTimestamp = CurrentTime;
+			UE_LOG(LogRivermax, Verbose, TEXT("Stream %d (%s:%u) stats: FrameCount: %llu, EndOfFrame: %llu, Chunks: %llu, Bytes: %llu, PacketLossInFrame: %llu, TotalPacketLoss: %llu, BiggerFrames: %llu, InvalidFrames: %llu, InvalidHeader: %llu, EmptyCompletion: %llu")
+			, StreamId
+			, *Options.StreamAddress
+			, Options.Port
+			, StreamStats.FramesReceived
+			, StreamStats.EndOfFrameReceived
+			, StreamStats.ChunksReceived
+			, StreamStats.BytesReceived
+			, StreamStats.FramePacketLossCount
+			, StreamStats.TotalPacketLossCount
+			, StreamStats.BiggerFramesCount
+			, StreamStats.InvalidFramesCount
+			, StreamStats.InvalidHeadercount
+			, StreamStats.EmptyCompletionCount
+			);
+
+			UE_LOG(LogRivermax, Verbose, TEXT("Stream %d (%s:%u) first packet interval: - Min: %llu. Max: %llu. Avg: %llu.")
+				, StreamId
+				, *Options.StreamAddress
+				, Options.Port
+				, StreamStats.MinFirstPacketIntervalNS
+				, StreamStats.MaxFirstPacketIntervalNS
+				, StreamStats.FirstPacketIntervalAccumulatorNS);
+
+		}
+	}
+
+	void FRivermaxInputStream::AllocateBuffers()
+	{
+		IRivermaxCoreModule& RivermaxModule = FModuleManager::LoadModuleChecked<IRivermaxCoreModule>(TEXT("RivermaxCore"));
+		if (RivermaxModule.GetRivermaxManager()->IsGPUDirectInputSupported() && Options.bUseGPUDirect)
+		{
+			bIsUsingGPUDirect = AllocateGPUBuffers();
+		}
+
+		constexpr uint32 CacheLineSize = PLATFORM_CACHE_LINE_SIZE;
+		if (bIsUsingGPUDirect == false)
+		{
+			BufferConfiguration.DataMemory->addr = FMemory::Malloc(BufferConfiguration.PayloadSize, CacheLineSize);
+		}
+		
+		BufferConfiguration.HeaderMemory->addr = FMemory::Malloc(BufferConfiguration.HeaderSize, CacheLineSize);
+
+		constexpr rmx_mkey_id InvalidKey = ((rmx_mkey_id)(-1L));
+		BufferConfiguration.DataMemory->mkey = InvalidKey;
+		BufferConfiguration.HeaderMemory->mkey = InvalidKey;
+	}
+
+	bool FRivermaxInputStream::AllocateGPUBuffers()
+	{
+		// Allocate memory space where rivermax input will write received buffer to
+
+		TRACE_CPUPROFILER_EVENT_SCOPE(FRivermaxInputStream::AllocateGPUBuffers);
+
+		const ERHIInterfaceType RHIType = RHIGetInterfaceType();
+		if (RHIType != ERHIInterfaceType::D3D12)
+		{
+			UE_LOG(LogRivermax, Warning, TEXT("Can't initialize input to use GPUDirect. RHI is %d but only Dx12 is supported at the moment."), int(RHIType));
+			return false;
+		}
+
+		FCUDAModule& CudaModule = FModuleManager::GetModuleChecked<FCUDAModule>("CUDA");
+		GPUDeviceIndex = CudaModule.GetCudaDeviceIndex();
+		CudaModule.DriverAPI()->cuCtxPushCurrent(CudaModule.GetCudaContextForDevice(GPUDeviceIndex));
+
+		// Todo: Add support for mgpu. 
+		CUdevice CudaDevice;
+		CUresult Status = CudaModule.DriverAPI()->cuDeviceGet(&CudaDevice, GPUDeviceIndex);
+		if (Status != CUDA_SUCCESS)
+		{
+			UE_LOG(LogRivermax, Warning, TEXT("Can't initialize input to use GPUDirect. Failed to get a Cuda device for GPU %d. Status: %d"), GPUDeviceIndex, Status);
+			return false;
+		}
+
+		CUmemAllocationProp AllocationProperties = {};
+		AllocationProperties.type = CU_MEM_ALLOCATION_TYPE_PINNED;
+		AllocationProperties.allocFlags.gpuDirectRDMACapable = 1; //is that required?
+		AllocationProperties.allocFlags.usage = 0;
+		AllocationProperties.location.type = CU_MEM_LOCATION_TYPE_DEVICE;
+		AllocationProperties.location.id = CudaDevice;
+
+		// Get memory granularity required for cuda device. We need to align allocation with this.
+		size_t Granularity;
+		Status = CudaModule.DriverAPI()->cuMemGetAllocationGranularity(&Granularity, &AllocationProperties, CU_MEM_ALLOC_GRANULARITY_RECOMMENDED);
+		if (Status != CUDA_SUCCESS)
+		{
+			UE_LOG(LogRivermax, Warning, TEXT("Can't initialize input to use GPUDirect. Failed to get allocation granularity. Status: %d"), Status);
+			return false;
+		}
+
+		// Cuda requires allocated memory to be aligned with a certain granularity
+		const size_t CudaAlignedAllocation = (BufferConfiguration.PayloadSize % Granularity) ? BufferConfiguration.PayloadSize + (Granularity - (BufferConfiguration.PayloadSize % Granularity)) : BufferConfiguration.PayloadSize;
+		
+		CUdeviceptr CudaBaseAddress;
+		constexpr CUdeviceptr InitialAddress = 0;
+		constexpr int32 Flags = 0;
+		Status = CudaModule.DriverAPI()->cuMemAddressReserve(&CudaBaseAddress, CudaAlignedAllocation, Granularity, InitialAddress, Flags);
+		if (Status != CUDA_SUCCESS)
+		{
+			UE_LOG(LogRivermax, Warning, TEXT("Can't initialize input to use GPUDirect. Failed to reserve memory for %d bytes. Status: %d"), CudaAlignedAllocation, Status);
+			return false;
+		}
+
+		// Make the allocation on device memory
+		CUmemGenericAllocationHandle Handle;
+		Status = CudaModule.DriverAPI()->cuMemCreate(&Handle, CudaAlignedAllocation, &AllocationProperties, Flags);
+		if (Status != CUDA_SUCCESS)
+		{
+			UE_LOG(LogRivermax, Warning, TEXT("Can't initialize input to use GPUDirect. Failed to create memory on device. Status: %d"), Status);
+			return false;
+		}
+		UE_LOG(LogRivermax, Verbose, TEXT("Allocated %d cuda memory"), CudaAlignedAllocation);
+
+		bool bExit = false;
+		constexpr int32 Offset = 0;
+		Status = CudaModule.DriverAPI()->cuMemMap(CudaBaseAddress, CudaAlignedAllocation, Offset, Handle, Flags);
+		if (Status != CUDA_SUCCESS)
+		{
+			UE_LOG(LogRivermax, Warning, TEXT("Can't initialize input to use GPUDirect. Failed to map memory. Status: %d"), Status);
+			// Need to release handle no matter what
+			bExit = true;
+		}
+
+		GPUAllocatedMemorySize = CudaAlignedAllocation;
+		GPUAllocatedMemoryBaseAddress = reinterpret_cast<void*>(CudaBaseAddress);
+
+		Status = CudaModule.DriverAPI()->cuMemRelease(Handle);
+		if (Status != CUDA_SUCCESS)
+		{
+			UE_LOG(LogRivermax, Warning, TEXT("Can't initialize input to use GPUDirect. Failed to release handle. Status: %d"), Status);
+			return false;
+		}
+
+		if (bExit)
+		{
+			return false;
+		}
+
+		// Setup access description.
+		CUmemAccessDesc MemoryAccessDescription = {};
+		MemoryAccessDescription.flags = CU_MEM_ACCESS_FLAGS_PROT_READWRITE;
+		MemoryAccessDescription.location.type = CU_MEM_LOCATION_TYPE_DEVICE;
+		MemoryAccessDescription.location.id = CudaDevice;
+		constexpr int32 Count = 1;
+		Status = CudaModule.DriverAPI()->cuMemSetAccess(CudaBaseAddress, CudaAlignedAllocation, &MemoryAccessDescription, Count);
+		if (Status != CUDA_SUCCESS)
+		{
+			UE_LOG(LogRivermax, Warning, TEXT("Can't initialize input to use GPUDirect. Failed to configure memory access. Status: %d"), Status);
+			return false;
+		}
+
+		CUstream CudaStream;
+		Status = CudaModule.DriverAPI()->cuStreamCreate(&CudaStream, CU_STREAM_NON_BLOCKING);
+		if (Status != CUDA_SUCCESS)
+		{
+			UE_LOG(LogRivermax, Warning, TEXT("Can't initialize input to use GPUDirect. Failed to create its stream. Status: %d"), Status);
+			return false;
+		}
+
+		GPUStream = CudaStream;
+
+		Status = CudaModule.DriverAPI()->cuCtxSynchronize();
+		if (Status != CUDA_SUCCESS)
+		{
+			UE_LOG(LogRivermax, Warning, TEXT("Can't initialize input to use GPUDirect. Failed to synchronize context. Status: %d"), Status);
+			return false;
+		}
+
+		// Give rivermax input buffer config the pointer to gpu allocated memory
+		BufferConfiguration.DataMemory->addr = GPUAllocatedMemoryBaseAddress;
+
+		CallbackPayload = MakeShared<FCallbackPayload>();
+
+		return true;
+	}
+
+	void FRivermaxInputStream::DeallocateBuffers()
+	{
+		if (GPUAllocatedMemorySize > 0)
+		{
+			FCUDAModule& CudaModule = FModuleManager::GetModuleChecked<FCUDAModule>("CUDA");
+			CudaModule.DriverAPI()->cuCtxPushCurrent(CudaModule.GetCudaContextForDevice(GPUDeviceIndex));
+
+			const CUdeviceptr BaseAddress = reinterpret_cast<CUdeviceptr>(GPUAllocatedMemoryBaseAddress);
+			CUresult Status = CudaModule.DriverAPI()->cuMemUnmap(BaseAddress, GPUAllocatedMemorySize);
+			if (Status != CUDA_SUCCESS)
+			{
+				UE_LOG(LogRivermax, Warning, TEXT("Failed to unmap cuda memory used for input stream. Status: %d"), Status);
+			}
+
+			Status = CudaModule.DriverAPI()->cuMemAddressFree(BaseAddress, GPUAllocatedMemorySize);
+			if (Status != CUDA_SUCCESS)
+			{
+				UE_LOG(LogRivermax, Warning, TEXT("Failed to free cuda memory used for input stream. Status: %d"), Status);
+			}
+			UE_LOG(LogRivermax, Verbose, TEXT("Deallocated %d cuda memory at address %d"), GPUAllocatedMemorySize, GPUAllocatedMemoryBaseAddress);
+
+			GPUAllocatedMemorySize = 0;
+			GPUAllocatedMemoryBaseAddress = 0;
+
+
+			for (const TPair<FRHIBuffer*, void*>& Entry : BufferGPUMemoryMap)
+			{
+				if (Entry.Value)
+				{
+					CudaModule.DriverAPI()->cuMemFree(reinterpret_cast<CUdeviceptr>(Entry.Value));
+				}
+			}
+			BufferGPUMemoryMap.Empty();
+
+			Status = CudaModule.DriverAPI()->cuStreamDestroy(reinterpret_cast<CUstream>(GPUStream));
+			if (Status != CUDA_SUCCESS)
+			{
+				UE_LOG(LogRivermax, Warning, TEXT("Failed to destroy cuda stream. Status: %d"), Status);
+			}
+			GPUStream = nullptr;
+
+
+			CudaModule.DriverAPI()->cuCtxPopCurrent(nullptr);
+		}
+	}
+
+	void* FRivermaxInputStream::GetMappedBuffer(const FBufferRHIRef& InBuffer)
+	{
+		// If we are here, d3d12 had to have been validated
+		const ERHIInterfaceType RHIType = RHIGetInterfaceType();
+		check(RHIType == ERHIInterfaceType::D3D12);
+
+		//Do we already have a mapped address for this buffer
+		if (BufferGPUMemoryMap.Find((InBuffer)) == nullptr)
+		{
+			int64 BufferMemorySize = 0;
+			CUexternalMemory MappedExternalMemory = nullptr;
+			HANDLE D3D12BufferHandle = 0;
+			CUDA_EXTERNAL_MEMORY_HANDLE_DESC CudaExtMemHandleDesc = {};
+
+			// Create shared handle for our buffer
+			{
+				TRACE_CPUPROFILER_EVENT_SCOPE(RmaxInput_D3D12CreateSharedHandle);
+
+				ID3D12Resource* NativeD3D12Resource = GetID3D12DynamicRHI()->RHIGetResource(InBuffer);
+				BufferMemorySize = GetID3D12DynamicRHI()->RHIGetResourceMemorySize(InBuffer);
+
+				TRefCountPtr<ID3D12Device> OwnerDevice;
+				HRESULT QueryResult;
+				if ((QueryResult = NativeD3D12Resource->GetDevice(IID_PPV_ARGS(OwnerDevice.GetInitReference()))) != S_OK)
+				{
+					UE_LOG(LogRivermax, Error, TEXT("Failed to get D3D12 device for captured buffer ressource: %d)"), QueryResult);
+					return nullptr;
+				}
+
+				if ((QueryResult = OwnerDevice->CreateSharedHandle(NativeD3D12Resource, NULL, GENERIC_ALL, NULL, &D3D12BufferHandle)) != S_OK)
+				{
+					UE_LOG(LogRivermax, Error, TEXT("Failed to create shared handle for captured buffer ressource: %d"), QueryResult);
+					return nullptr;
+				}
+
+				CudaExtMemHandleDesc.type = CU_EXTERNAL_MEMORY_HANDLE_TYPE_D3D12_RESOURCE;
+				CudaExtMemHandleDesc.handle.win32.name = nullptr;
+				CudaExtMemHandleDesc.handle.win32.handle = D3D12BufferHandle;
+				CudaExtMemHandleDesc.size = BufferMemorySize;
+				CudaExtMemHandleDesc.flags |= CUDA_EXTERNAL_MEMORY_DEDICATED;
+			}
+
+			FCUDAModule& CudaModule = FModuleManager::GetModuleChecked<FCUDAModule>("CUDA");
+
+			CudaModule.DriverAPI()->cuCtxPushCurrent(CudaModule.GetCudaContext());
+
+			{
+				TRACE_CPUPROFILER_EVENT_SCOPE(Rmax_CudaImportMemory);
+
+				const CUresult Result = FCUDAModule::CUDA().cuImportExternalMemory(&MappedExternalMemory, &CudaExtMemHandleDesc);
+
+				if (D3D12BufferHandle)
+				{
+					CloseHandle(D3D12BufferHandle);
+				}
+
+				if (Result != CUDA_SUCCESS)
+				{
+					UE_LOG(LogRivermax, Error, TEXT("Failed to import shared buffer. Error: %d"), Result);
+					return nullptr;
+				}
+			}
+
+			{
+				TRACE_CPUPROFILER_EVENT_SCOPE(Rmax_MapCudaMemory);
+
+				CUDA_EXTERNAL_MEMORY_BUFFER_DESC BufferDescription = {};
+				BufferDescription.offset = 0;
+				BufferDescription.size = BufferMemorySize;
+				CUdeviceptr NewMemory;
+				const CUresult Result = FCUDAModule::CUDA().cuExternalMemoryGetMappedBuffer(&NewMemory, MappedExternalMemory, &BufferDescription);
+				if (Result != CUDA_SUCCESS || NewMemory == 0)
+				{
+					UE_LOG(LogRivermax, Error, TEXT("Failed to get shared buffer mapped memory. Error: %d"), Result);
+					CudaModule.DriverAPI()->cuCtxPushCurrent(nullptr);
+					return nullptr;
+				}
+
+				BufferGPUMemoryMap.Add(InBuffer, reinterpret_cast<void*>(NewMemory));
+			}
+
+			CudaModule.DriverAPI()->cuCtxPushCurrent(nullptr);
+		}
+
+		// At this point, we have the mapped buffer in cuda space and we can use it to schedule a memcpy on cuda engine.
+		return BufferGPUMemoryMap[InBuffer];
+	}
+
+	void FRivermaxInputStream::ProcessSRD(const FRTPHeader& RTPHeader, const uint8* DataPtr)
+	{
+		if (RTPHeader.SRD1.Length <= 0)
+		{
+			return;
+		}
+
+		uint32 DataOffset = 0;
+		uint32 PayloadSize = RTPHeader.SRD1.Length;
+		if (RTPHeader.SRD1.bHasContinuation)
+		{
+			PayloadSize += RTPHeader.SRD2.Length;
+		}
+
+		if (bIsUsingGPUDirect)
+		{
+			// Initial case (start address)
+			if (StreamData.DeviceWritePointerOne == nullptr)
+			{
+				StreamData.DeviceWritePointerOne = DataPtr;
+				StreamData.SizeToWriteOne = PayloadSize;
+			}
+			else
+			{
+				// Detection of wrap around -> Move tracking to second buffer
+				if (StreamData.DeviceWritePointerTwo == nullptr && DataPtr < StreamData.DeviceWritePointerOne)
+				{
+					StreamData.DeviceWritePointerTwo = DataPtr;
+					StreamData.SizeToWriteTwo = 0;
+				}
+
+				// Case where we track memory in first buffer
+				if (StreamData.DeviceWritePointerTwo == nullptr)
+				{
+					StreamData.SizeToWriteOne += PayloadSize;
+				}
+				else // Tracking memory in second buffer
+				{
+					StreamData.SizeToWriteTwo += PayloadSize;
+				}
+			}
+		}
+		else
+		{
+			uint8* WriteBuffer = reinterpret_cast<uint8*>(StreamData.CurrentFrameVideoBuffer);
+			FMemory::Memcpy(&WriteBuffer[StreamData.WritingOffset], &DataPtr[DataOffset], RTPHeader.SRD1.Length);
+			StreamData.WritingOffset += RTPHeader.SRD1.Length;
+
+			if (RTPHeader.SRD1.bHasContinuation)
+			{
+				DataOffset += RTPHeader.SRD1.Length;
+				FMemory::Memcpy(&WriteBuffer[StreamData.WritingOffset], &DataPtr[DataOffset], RTPHeader.SRD2.Length);
+				StreamData.WritingOffset += RTPHeader.SRD2.Length;
+			}
+		}
+
+		StreamData.ReceivedSize += PayloadSize;
+	}
+
+	void FRivermaxInputStream::ProcessLastSRD(const FRTPHeader& RTPHeader, const uint8* DataPtr, TSharedPtr<IRivermaxSample> InSample)
+	{
+		if (StreamData.ReceivedSize == StreamData.ExpectedSize)
+		{
+			++StreamStats.FramesReceived;
+			uint32 FrameNumber = Utils::TimestampToFrameNumber(RTPHeader.Timestamp, Options.FrameRate);
+			TRACE_CPUPROFILER_EVENT_SCOPE_TEXT(*FRivermaxTracingUtils::RmaxInReceivedFrameTraceEvents[FrameNumber % 10]);
+
+			if (bIsUsingGPUDirect)
+			{
+				// For gpudirect, we need all payload sizes to be equal with the exception of the last one.
+				if (FrameDescriptionTracking.CommonPayloadSize <= 0)
+				{
+					UE_LOG(LogRivermax, Warning, TEXT("Unsupported variable SRD length detected while GPUDirect for input stream is used. Disable and reopen the stream."));
+					Listener->OnStreamError();
+					bIsShuttingDown = true;
+				}
+
+				//Frame received entirely, time to copy it from rivermax gpu scratchpad to our own gpu memory
+				FCUDAModule& CudaModule = FModuleManager::GetModuleChecked<FCUDAModule>("CUDA");
+				CUresult Result = CudaModule.DriverAPI()->cuCtxPushCurrent(CudaModule.GetCudaContext());
+
+				const CUdeviceptr DestinationGPUMemory = reinterpret_cast<CUdeviceptr>(StreamData.CurrentFrameVideoBuffer);
+				const CUdeviceptr SourceGPUMemoryOne = reinterpret_cast<CUdeviceptr>(StreamData.DeviceWritePointerOne);
+
+				const uint32 NumSRDPartOne = StreamData.SizeToWriteOne / FrameDescriptionTracking.CommonPayloadSize;
+				const uint32 NumSRDPartTwo = StreamData.SizeToWriteTwo / FrameDescriptionTracking.CommonPayloadSize;
+
+				// Use cuda's 2d memcopy to do a source and destination stride difference memcopy
+				// We initialize rivermax stream with a payload size blindly since we don't know what the sender will use
+				// So, we use a big value by default and expect SRD to be smaller
+				// This memcopy will consume the SRD size value but jump the init payload size value on the source address
+				// Limitation is that this will only work for fixed SRD across a frame
+				CUDA_MEMCPY2D StrideDescription;
+				FMemory::Memset(StrideDescription, 0);
+				StrideDescription.srcDevice = SourceGPUMemoryOne;
+				StrideDescription.dstDevice = DestinationGPUMemory;
+				StrideDescription.dstMemoryType = CUmemorytype::CU_MEMORYTYPE_DEVICE;
+				StrideDescription.srcMemoryType = CUmemorytype::CU_MEMORYTYPE_DEVICE;
+				StrideDescription.srcPitch = ExpectedPayloadSize; //Source pitch is the expected payload used at init
+				StrideDescription.dstPitch = FrameDescriptionTracking.CommonPayloadSize; //Destination pitch is the fixed SRD size we received
+				StrideDescription.WidthInBytes = FrameDescriptionTracking.CommonPayloadSize; //Width in bytes is the amount to copy, the SRD size
+				StrideDescription.Height = NumSRDPartOne;
+				Result = CudaModule.DriverAPI()->cuMemcpy2DAsync(&StrideDescription, reinterpret_cast<CUstream>(GPUStream));
+
+				// Keep track where we are at in the source buffer
+				CUdeviceptr SourcePtr = SourceGPUMemoryOne + (ExpectedPayloadSize * NumSRDPartOne);
+
+				if (StreamData.DeviceWritePointerTwo != nullptr && StreamData.SizeToWriteTwo > 0)
+				{
+					StrideDescription.srcDevice = reinterpret_cast<CUdeviceptr>(StreamData.DeviceWritePointerTwo);
+					StrideDescription.dstDevice = reinterpret_cast<CUdeviceptr>(StreamData.CurrentFrameVideoBuffer) + StreamData.SizeToWriteOne;
+					StrideDescription.Height = NumSRDPartTwo;
+					Result = CudaModule.DriverAPI()->cuMemcpy2DAsync(&StrideDescription, reinterpret_cast<CUstream>(GPUStream));
+
+					// Update source pointer since we used the second pointer tracker
+					SourcePtr = StrideDescription.srcDevice + (ExpectedPayloadSize * NumSRDPartTwo);
+				}
+
+				// Verify if we copied whole frame using the stride 2d copy or we're missing the last SRD of fractional size
+				const uint32 TotalMemoryCopied = (FrameDescriptionTracking.CommonPayloadSize * (NumSRDPartOne + NumSRDPartTwo));
+				if (TotalMemoryCopied < StreamData.ExpectedSize)
+				{
+					const uint32 SizeLeftToCopy = StreamData.ExpectedSize - TotalMemoryCopied;
+					if (ensure(SizeLeftToCopy == RTPHeader.SRD1.Length))
+					{
+						CUdeviceptr DestinationPtr = DestinationGPUMemory + TotalMemoryCopied;
+						Result = CudaModule.DriverAPI()->cuMemcpyDtoDAsync(DestinationPtr, SourcePtr, SizeLeftToCopy, reinterpret_cast<CUstream>(GPUStream));
+					}
+				}
+
+				if (Result != CUDA_SUCCESS)
+				{
+					UE_LOG(LogRivermax, Warning, TEXT("Failed to copy received buffer to shared memory. Error: %d"), Result);
+					Listener->OnVideoFrameReceptionError(StaticCastSharedPtr<IRivermaxVideoSample>(InSample));
+					State = EReceptionState::FrameError;
+					FrameErrorState(RTPHeader);
+					return;
+				}
+
+				auto CudaCallback = [](void* userData)
+				{
+					TRACE_CPUPROFILER_EVENT_SCOPE(FRivermaxInputStream::MemcopyCallback);
+					if (userData)
+					{
+						// It might happen that our stream has been closed once the callback is triggered
+						TWeakPtr<FCallbackPayload>* WeakPayloadPtr = reinterpret_cast<TWeakPtr<FCallbackPayload>*>(userData);
+						TWeakPtr<FCallbackPayload> WeakPayload = *WeakPayloadPtr;
+						if (TSharedPtr<FCallbackPayload> Payload = WeakPayload.Pin())
+						{
+							Payload->bIsWaitingForPendingCopy = false;
+						}
+					}
+				};
+
+				// Schedule a callback to know when to make the frame available
+				CallbackPayload->bIsWaitingForPendingCopy = true;
+				TWeakPtr<FCallbackPayload> WeakPayload = CallbackPayload;
+				CudaModule.DriverAPI()->cuLaunchHostFunc(reinterpret_cast<CUstream>(GPUStream), CudaCallback, &WeakPayload);
+
+				FCUDAModule::CUDA().cuCtxPopCurrent(nullptr);
+
+				// For now, we wait for the cuda callback before we move on receiving next frame. 
+				// Will need to update this and make the frame available from the cuda callback to avoid losing packets 
+				{
+					TRACE_CPUPROFILER_EVENT_SCOPE(FRivermaxInputStream:WaitingPendingOperation)
+					const double CallbackTimestamp = FPlatformTime::Seconds();
+					while (CallbackPayload->bIsWaitingForPendingCopy == true && bIsShuttingDown == false)
+					{
+						FPlatformProcess::SleepNoStats(UE::RivermaxCore::Private::Utils::SleepTimeSeconds);
+						if (FPlatformTime::Seconds() - CallbackTimestamp > CVarWaitForCompletionTimeout.GetValueOnAnyThread())
+						{
+							UE_LOG(LogRivermax, Error, TEXT("Waiting for gpu copy of sample timed out."));
+							Listener->OnStreamError();
+							CallbackPayload->bIsWaitingForPendingCopy = false;
+							break;
+						}
+					}
+				}
+			}
+
+			// Finished processing incoming frame, reset tracking data associated with it
+			FrameDescriptionTracking.ResetSingleFrameTracking();
+
+			// No need to provide the new frame and prepare the next one if we are shutting down
+			if (bIsShuttingDown == false)
+			{
+				FRivermaxInputVideoFrameReception NewFrame;
+				NewFrame.VideoBuffer = reinterpret_cast<uint8*>(StreamData.CurrentFrameVideoBuffer);
+				Listener->OnVideoFrameReceived(StaticCastSharedPtr<IRivermaxVideoSample>(InSample));
+				StreamData.CurrentFrameVideoBuffer = nullptr;
+				StreamData.CurrentSamples.Empty();
+				// Finished receiving a frame, move back to reception
+				State = EReceptionState::Receiving;
+			}
+		}
+		else
+		{
+			UE_LOG(LogRivermax, Warning, TEXT("End of frame received but not enough data was received (missing %d). Expected %d but received (%d)"), StreamData.ExpectedSize - StreamData.ReceivedSize, StreamData.ExpectedSize, StreamData.ReceivedSize);
+			
+			Listener->OnVideoFrameReceptionError(StaticCastSharedPtr<IRivermaxVideoSample>(InSample));
+			State = EReceptionState::FrameError;
+			FrameErrorState(RTPHeader);
+			
+			++StreamStats.InvalidFramesCount;
+		}
+	}
+
+	void FRivermaxInputStream::FrameReceptionState(const FRTPHeader& RTPHeader, const uint8* DataPtr, TSharedPtr<IRivermaxSample> InSample)
+	{
+		const uint64 LastSequenceNumberIncremented = StreamData.LastSequenceNumber + 1;
+
+		const uint64 LostPackets = ((uint64)RTPHeader.SequenceNumber + 0x100000000 - LastSequenceNumberIncremented) & 0xFFFFFFFF;
+		if (LostPackets > 0)
+		{
+			UE_LOG(LogRivermax, Warning, TEXT("Lost %llu packets during reception of chunk"), LostPackets);
+			
+			StreamData.WritingOffset = 0;
+			StreamData.ReceivedSize = 0;
+			++StreamStats.TotalPacketLossCount;
+			++StreamStats.FramePacketLossCount;
+
+			// For now, if packets were lost, skip the incoming frame. We could improve that and have corrupted frames instead of skipping them but can be added later
+			Listener->OnVideoFrameReceptionError(StaticCastSharedPtr<IRivermaxVideoSample>(InSample));
+			State = EReceptionState::FrameError;
+			FrameErrorState(RTPHeader);
+			return;
+		}
+
+		StreamData.LastSequenceNumber = RTPHeader.SequenceNumber;
+
+		ProcessSRD(RTPHeader, DataPtr);
+
+		if (StreamData.ReceivedSize > StreamData.ExpectedSize)
+		{
+			UE_LOG(LogRivermax, Warning, TEXT("Received too much data (%d). Expected %d but received (%d)"), StreamData.ReceivedSize - StreamData.ExpectedSize, StreamData.ExpectedSize, StreamData.ReceivedSize);
+			StreamData.WritingOffset = 0;
+			StreamData.ReceivedSize = 0;
+			++StreamStats.BiggerFramesCount;
+
+			Listener->OnVideoFrameReceptionError(StaticCastSharedPtr<IRivermaxVideoSample>(InSample));
+			State = EReceptionState::FrameError;
+			FrameErrorState(RTPHeader);
+			return;
+		}
+
+		if (RTPHeader.bIsMarkerBit)
+		{
+			ProcessLastSRD(RTPHeader, DataPtr, InSample);
+
+			StreamStats.FramePacketLossCount = 0;
+			++StreamStats.EndOfFrameReceived;
+		}
+	}
+
+	void FRivermaxInputStream::WaitForMarkerState(const FRTPHeader& RTPHeader)
+	{
+		if (RTPHeader.bIsMarkerBit)
+		{
+			StreamData.LastSequenceNumber = RTPHeader.SequenceNumber;
+			StreamData.WritingOffset = 0;
+			StreamData.ReceivedSize = 0;
+			StreamData.DeviceWritePointerOne = nullptr;
+			StreamData.SizeToWriteOne = 0;
+			StreamData.DeviceWritePointerTwo = nullptr;
+			StreamData.SizeToWriteTwo = 0;
+			bIsFirstPacketReceived = false;
+
+			StreamData.CurrentFrameVideoBuffer = nullptr;
+			StreamData.CurrentSamples.Empty();
+			State = EReceptionState::Receiving;
+		}
+	}
+
+	void FRivermaxInputStream::FrameErrorState(const FRTPHeader& RTPHeader)
+	{
+		// In error, we just wait for the end of frame (marker bit) to restart reception
+		WaitForMarkerState(RTPHeader);
+	}
+
+	void FRivermaxInputStream::UpdateFrameTracking(const FRTPHeader& NewRTPHeader)
+	{
+		FrameDescriptionTracking.EvaluateNewRTP(NewRTPHeader);
+		
+		if (!Options.bEnforceVideoFormat)
+		{
+			if (FrameDescriptionTracking.DetectedResolution.X > 0 && FrameDescriptionTracking.DetectedResolution.Y > 0)
+			{
+				if (FrameDescriptionTracking.DetectedResolution != StreamResolution)
+				{
+					StreamResolution = FrameDescriptionTracking.DetectedResolution;
+
+					FRivermaxInputVideoFormatChangedInfo FormatChangeInfo;
+					FormatChangeInfo.Width = StreamResolution.X;
+					FormatChangeInfo.Height = StreamResolution.Y;
+					FormatChangeInfo.PixelFormat = FrameDescriptionTracking.ExpectedSamplingType;
+					Listener->OnVideoFormatChanged(FormatChangeInfo);
+				}
+			}
+		}
+	}
+
+	bool FRivermaxInputStream::InitializeNetworkSettings(FNetworkSettings& OutSettings, FString& OutErrorMessage)
+	{
+		//Retrieves all network data from string version of device interface and stream address
+
+		memset(&OutSettings.DeviceInterface, 0, sizeof(OutSettings.DeviceInterface));
+		if (inet_pton(AF_INET, StringCast<ANSICHAR>(*Options.InterfaceAddress).Get(), &OutSettings.DeviceInterface.sin_addr) != 1)
+		{
+			OutErrorMessage = FString::Printf(TEXT("Failed to convert Device interface '%s' to network address"), *Options.InterfaceAddress);
+			return false;
+		}
+
+		memset(&OutSettings.DestinationAddress, 0, sizeof(OutSettings.DestinationAddress));
+		FMemory::Memset(&OutSettings.DestinationAddress, 0, sizeof(OutSettings.DestinationAddress));
+		if (inet_pton(AF_INET, StringCast<ANSICHAR>(*Options.StreamAddress).Get(), &OutSettings.DestinationAddress.sin_addr) != 1)
+		{
+			OutErrorMessage = FString::Printf(TEXT("Failed to convert Stream address '%s' to network address"), *Options.StreamAddress);
+			return false;
+		}
+
+		OutSettings.DeviceInterface.sin_family = AF_INET;
+		OutSettings.DestinationAddress.sin_family = AF_INET;
+		OutSettings.DestinationAddress.sin_port = ByteSwap((uint16)Options.Port);
+
+		OutSettings.DeviceAddress.family = AF_INET;
+		OutSettings.DeviceAddress.addr.ipv4 = OutSettings.DeviceInterface.sin_addr;
+		rmx_status Status = CachedAPI->rmx_retrieve_device_iface(&OutSettings.RMXDeviceInterface, &OutSettings.DeviceAddress);
+		if(Status != RMX_OK)
+		{
+			OutErrorMessage = FString::Printf(TEXT("Could not retrieve Rivermax interface for IP '%s'"), *Options.InterfaceAddress);
+			return false;
+		}
+		
+		CachedAPI->rmx_input_set_stream_nic_address(&StreamParameters, reinterpret_cast<sockaddr*>(&OutSettings.DeviceInterface));
+
+		return true;
+	}
+
+	bool FRivermaxInputStream::InitializeStreamParameters(FString& OutErrorMessage)
+	{
+		size_t BufferElement = 1 << 18;
+		CachedAPI->rmx_input_set_mem_capacity_in_packets(&StreamParameters, BufferElement);
+
+		// We always split header and data so we always have two sub blocks (Header + Payload)
+		constexpr size_t SubBlockCount = 2;
+		CachedAPI->rmx_input_set_mem_sub_block_count(&StreamParameters, SubBlockCount);
+		CachedAPI->rmx_input_set_entry_size_range(&StreamParameters, BufferConfiguration.HeaderBlockID, BufferConfiguration.HeaderExpectedSize, BufferConfiguration.HeaderExpectedSize);
+		CachedAPI->rmx_input_set_entry_size_range(&StreamParameters, BufferConfiguration.DataBlockID, ExpectedPayloadSize, ExpectedPayloadSize);
+
+		constexpr rmx_input_option InputOptions = RMX_INPUT_STREAM_CREATE_INFO_PER_PACKET;
+		constexpr rmx_input_timestamp_format TimestampFormat = RMX_INPUT_TIMESTAMP_RAW_NANO;
+		CachedAPI->rmx_input_enable_stream_option(&StreamParameters, InputOptions);
+		CachedAPI->rmx_input_set_timestamp_format(&StreamParameters, TimestampFormat);
+
+		if (bIsDynamicHeaderEnabled)
+		{
+			CachedAPI->rmx_input_enable_stream_option(&StreamParameters, RMX_INPUT_STREAM_RTP_SMPTE_2110_20_DYNAMIC_HDS);
+		}
+
+		const rmx_status Status = CachedAPI->rmx_input_determine_mem_layout(&StreamParameters);
+		if (Status != RMX_OK)
+		{
+			OutErrorMessage = FString::Printf(TEXT("Could not determine memory layout for input stream using IP %s. Status: %d."), *Options.InterfaceAddress, Status);
+			return false;
+		}
+
+		BufferElement = CachedAPI->rmx_input_get_mem_capacity_in_packets(&StreamParameters);
+		BufferConfiguration.DataMemory = CachedAPI->rmx_input_get_mem_block_buffer(&StreamParameters, BufferConfiguration.DataBlockID);
+		BufferConfiguration.HeaderMemory = CachedAPI->rmx_input_get_mem_block_buffer(&StreamParameters, BufferConfiguration.HeaderBlockID);
+
+		if (BufferConfiguration.HeaderMemory->length <= 0)
+		{
+			OutErrorMessage = FString::Printf(TEXT("Header data split not supported for device with IP %s. Can't initialize stream."), *Options.InterfaceAddress);
+			return false;
+		}
+		
+		BufferConfiguration.PayloadSize = BufferConfiguration.DataMemory->length;
+		BufferConfiguration.HeaderSize = BufferConfiguration.HeaderMemory->length;
+
+		AllocateBuffers();
+
+		return true;
+	}
+
+	bool FRivermaxInputStream::FinalizeStreamCreation(const FNetworkSettings& NetworkSettings, FString& OutErrorMessage)
+	{
+		rmx_status Status = CachedAPI->rmx_input_create_stream(&StreamParameters, &StreamId);
+		if (Status != RMX_OK)
+		{
+			OutErrorMessage = FString::Printf(TEXT("Could not create stream. Status: %d."), Status);
+			return false;
+		}
+
+		// Once stream creation is done, retrieve strides it is using
+		BufferConfiguration.DataStride = CachedAPI->rmx_input_get_stride_size(&StreamParameters, BufferConfiguration.DataBlockID);
+		BufferConfiguration.HeaderStride = CachedAPI->rmx_input_get_stride_size(&StreamParameters, BufferConfiguration.HeaderBlockID);
+
+		// Configure how we wnat to consume chunks
+		constexpr size_t MinChunkSize = 0;
+		constexpr size_t MaxChunkSize = 5000;
+		constexpr int Timeout = 0;
+		Status = CachedAPI->rmx_input_set_completion_moderation(StreamId, MinChunkSize, MaxChunkSize, Timeout);
+		if (Status != RMX_OK)
+		{
+			OutErrorMessage = FString::Printf(TEXT("Could not setup expected packet count for stream. Status: %d."), Status);
+			return false;
+		}
+
+		// Finalize stream creation
+		CachedAPI->rmx_input_init_chunk_handle(&BufferConfiguration.ChunkHandle, StreamId);
+		CachedAPI->rmx_input_init_flow(&FlowAttribute);
+		CachedAPI->rmx_input_set_flow_local_addr(&FlowAttribute, reinterpret_cast<const sockaddr*>(&NetworkSettings.DestinationAddress));
+
+		sockaddr_in SourceAddr;
+		memset(&SourceAddr, 0, sizeof(SourceAddr));
+		SourceAddr.sin_family = AF_INET;
+		CachedAPI->rmx_input_set_flow_remote_addr(&FlowAttribute, reinterpret_cast<const sockaddr*>(&SourceAddr));
+		CachedAPI->rmx_input_set_flow_tag(&FlowAttribute, FlowTag);
+
+		Status = CachedAPI->rmx_input_attach_flow(StreamId, &FlowAttribute);
+		if(Status != RMX_OK)
+		{
+			OutErrorMessage = FString::Printf(TEXT("Could not attach flow to stream. Status: %d."), Status);
+
+			// Cleanup stream.
+			Status = CachedAPI->rmx_input_destroy_stream(StreamId);
+
+			if (Status != RMX_OK)
+			{
+				UE_LOG(LogRivermax, Warning, TEXT("Failed to destroy input stream %d correctly. Status: %d"), StreamId, Status);
+			}
+
+			return false;
+		}
+		
+		return true;
+	}
+}
+
+	
+
diff --git a/Engine/Plugins/VirtualProduction/Rivermax/RivermaxCore/Source/RivermaxCore/Private/Streams/RivermaxInputStream.h b/Engine/Plugins/VirtualProduction/Rivermax/RivermaxCore/Source/RivermaxCore/Private/Streams/RivermaxInputStream.h
new file mode 100644
--- /dev/null
+++ b/Engine/Plugins/VirtualProduction/Rivermax/RivermaxCore/Source/RivermaxCore/Private/Streams/RivermaxInputStream.h
@@ -0,0 +1,329 @@
+// Copyright Epic Games, Inc. All Rights Reserved.
+
+#pragma once
+
+#include "IRivermaxInputStream.h"
+
+#include <atomic>
+#include "Async/Future.h"
+#include "HAL/Runnable.h"
+#include "RHI.h"
+#include "RivermaxWrapper.h"
+#include "RivermaxTypes.h"
+
+
+
+namespace UE::RivermaxCore::Private
+{
+	using UE::RivermaxCore::IRivermaxInputStream;
+	using UE::RivermaxCore::IRivermaxInputStreamListener;
+	using UE::RivermaxCore::FRivermaxInputStreamOptions;
+
+	struct FRawRTPHeader;
+	struct FRTPHeader;
+	
+	struct FInputStreamBufferConfiguration
+	{
+		size_t PayloadSize = 0;
+		size_t HeaderSize = 0;
+		uint16 PayloadExpectedSize = 1500;
+		uint16 HeaderExpectedSize = 32;
+
+		/** Handle returned by Rivermax library to use the Chunk API */
+		rmx_input_chunk_handle ChunkHandle;
+
+		/** Data Sub block ID */
+		const uint8 HeaderBlockID = 0;
+
+		/** Data Sub block ID */
+		const uint8 DataBlockID = 1;
+
+		/** Memory region used to store incoming payloads */
+		rmx_mem_region* DataMemory = nullptr;
+		
+		/** Payload stride in bytes for input stream */
+		size_t DataStride = 0;
+		
+		/** Memory region used to store incoming headers */
+		rmx_mem_region* HeaderMemory = nullptr;
+
+		/** Header stride in bytes for input stream */
+		size_t HeaderStride = 0;
+	};
+
+	struct FInputStreamStats
+	{
+		uint64 InvalidHeadercount = 0;
+		uint64 FramePacketLossCount = 0;
+		uint64 TotalPacketLossCount = 0;
+		uint64 BiggerFramesCount = 0;
+		uint64 InvalidFramesCount = 0;
+		uint64 BytesReceived = 0;
+		uint64 FramesReceived = 0;
+		uint64 ChunksReceived = 0;
+		uint64 EndOfFrameReceived = 0;
+		uint64 EmptyCompletionCount = 0;
+		uint64 MinFirstPacketIntervalNS = TNumericLimits<uint64>::Max();
+		uint64 MaxFirstPacketIntervalNS = TNumericLimits<uint64>::Min();
+		uint64 FirstPacketIntervalAccumulatorNS = 0;
+		uint64 FirstPacketIntervalStatsCount = 0;
+	};
+
+	struct FInputStreamData
+	{
+		uint64 LastSequenceNumber = 0;
+		void* CurrentFrameVideoBuffer = nullptr;
+		uint32 WritingOffset = 0;
+		uint32 ReceivedSize = 0;
+		uint32 ExpectedSize = 0;
+
+		// GPUDirect data handling. Rivermax uses a cyclic buffer so we might have part of the frame at the end of the first buffer and the rest on the second part.
+		const void* DeviceWritePointerOne = nullptr;
+		const void* DeviceWritePointerTwo = nullptr;
+		uint32 SizeToWriteOne = 0;
+		uint32 SizeToWriteTwo = 0;
+
+		// A map of samples for the current frame. 
+		TMap<IRivermaxSample::ESampleType, TSharedPtr<IRivermaxSample>> CurrentSamples;
+	};
+
+	/** Incoming frame description tracker. Used to detect variable SRD length across a frame for now but should be used to help auto detect incoming frame */
+	struct FFrameDescriptionTrackingData
+	{
+	public:
+
+		/** Resets data related to current frame being tracked */
+		void ResetSingleFrameTracking();
+
+		/** Look at new RTP and extract info to estimate frame description */
+		void EvaluateNewRTP(const FRTPHeader& NewHeader);
+
+		/** Whether a valid resolution was detected during the last frame */
+		bool HasDetectedValidResolution() const;
+
+	private:
+		
+		/** Updates current resolution detection using new RTP header. */
+		void UpdateResolutionDetection(const FRTPHeader& NewHeader);
+		
+		/** Updates payload sizes encountered throughout a frame */
+		void UpdatePayloadSizeTracking(const FRTPHeader& NewHeader);
+
+	public:
+		/** Used to track different payload sizes received. GPUDirect can't work if different payload sizes are received across the frame except the last one */
+		TMap<uint32, uint32> PayloadSizeReceived;
+
+		/** Payload size used across a frame. -1 if multi SRD except last one are received. */
+		int16 CommonPayloadSize = -1;
+
+		/** Last resolution detected */
+		FIntPoint DetectedResolution = FIntPoint::ZeroValue;
+
+		/** Sampling type we are expecting as setup by the user. We could also auto detect it but some sampling types would collide so left manual for now. */
+		ESamplingType ExpectedSamplingType = ESamplingType::RGB_10bit;
+
+		/** Used to log sampling type warning once until a good frame type is received */
+		bool bHasLoggedSamplingWarning = false;
+	};
+
+	/** Struct holding different network settings in various forms to initialize Rivermax stream */
+	struct FNetworkSettings
+	{
+		sockaddr_in DeviceInterface;
+		rmx_device_iface RMXDeviceInterface;
+		rmx_ip_addr DeviceAddress;
+		sockaddr_in DestinationAddress;
+	};
+
+	class FRivermaxInputStream : public IRivermaxInputStream, public FRunnable
+	{
+	public:
+		FRivermaxInputStream();
+		virtual ~FRivermaxInputStream();
+
+	public:
+
+		//~ Begin IRivermaxInputStreamListener interface
+		virtual bool Initialize(const FRivermaxInputStreamOptions& InOptions, IRivermaxInputStreamListener& InListener) override;
+		virtual void Uninitialize() override;
+		//~ End IRivermaxInputStreamListener interface 
+
+		void Process_AnyThread();
+
+		//~ Begin FRunnable interface
+		virtual bool Init() override;
+		virtual uint32 Run() override;
+		virtual void Stop() override;
+		virtual void Exit() override;
+		//~ End FRunnable interface
+
+	private:
+
+		/** Process reception state. i.e Extract RTP, tracks data received, etc... */
+		void FrameReceptionState(const FRTPHeader& RTPHeader, const uint8* DataPtr, TSharedPtr<IRivermaxSample> InSample);
+
+		/** State where we just wait to receive an end of frame in RTP. Used to align for the first frame and when an error occurs during reception */
+		void WaitForMarkerState(const FRTPHeader& RTPHeader);
+
+		/** 
+		 * State we go in when an error was detected during the reception of a frame.  
+		 * Once the end of frame is received, will go back to frame reception state.
+		 */
+		void FrameErrorState(const FRTPHeader& RTPHeader);
+
+		/** 
+		 * Iterates over each chunks received in the Completion data structure.
+		 * For each item, its header is parsed and validated before processing data it contains
+		 */
+		void ParseChunks(const rmx_input_completion* Completion);
+
+		/** Process a SRD / packet by copying data associated to it into output buffer or tracking memory received for gpudirect */
+		void ProcessSRD(const FRTPHeader& RTPHeader, const uint8* DataPtr);
+
+		/** Handles wrapping up of a frame when marker bit is detected (last SRD). Will notify listener of a new frame received */
+		void ProcessLastSRD(const FRTPHeader& RTPHeader, const uint8* DataPtr, TSharedPtr<IRivermaxSample> InSample);
+
+		/** Requests a new frame from our listener. If it fails, it will fall in error state waiting for the next one */
+		TSharedPtr<IRivermaxVideoSample> GetVideoSampleForReception(const FRTPHeader& RTPHeader);
+
+		/** If enabled, will printout reception stats to logs */
+		void LogStats();
+
+		/** Allocates system memory buffer used for reception and used by rivermax */
+		void AllocateBuffers();
+
+		/** Allocates gpu memory buffer used for reception and used by rivermax */
+		bool AllocateGPUBuffers();
+
+		/** Cleans up allocated buffers */
+		void DeallocateBuffers();
+
+		/** Gets the mapped gpu buffer for the incoming RHI buffer. Caching will occur if a new buffer is given */
+		void* GetMappedBuffer(const FBufferRHIRef& InBuffer);
+
+		/** Updates frame information tracking using RTP header */
+		void UpdateFrameTracking(const FRTPHeader& NewRTPHeader);
+
+		/** Initializes network settings for desired input stream options */
+		bool InitializeNetworkSettings(FNetworkSettings& OutSettings, FString& OutErrorMessage);
+		
+		/** Configures input streams parameters*/
+		bool InitializeStreamParameters(FString& OutErrorMessage);
+
+		/** Finalize Rivermax stream initialization */
+		bool FinalizeStreamCreation(const FNetworkSettings& NetworkSettings, FString& OutErrorMessage);
+
+	private:
+
+		/** Options used for this stream, such has resolution, frame rate etc... */
+		FRivermaxInputStreamOptions Options;
+
+		/** Thread on which reception is handled */
+		FRunnableThread* RivermaxThread = nullptr;
+
+		/** Whether reception thread should keep running */
+		std::atomic<bool> bIsActive = false;
+
+		/** Stream id given back by rivermax */
+		rmx_stream_id StreamId = 0;
+
+		/** Holding stream configuration data */
+		rmx_input_stream_params StreamParameters;
+
+		/** Flow attributes used at stream creation. Used to compare incoming data to our flow and detect mismatches */
+		rmx_input_flow FlowAttribute;
+
+		/** Flow identification. Used to verify if packets belong to expected stream */
+		uint32 FlowTag = 0;
+
+		/** Memory configuration used for this stream. */
+		FInputStreamBufferConfiguration BufferConfiguration;
+
+		/** Used to know when the first packet of a frame has been received. Used mostly for visualization in insights */
+		bool bIsFirstPacketReceived = false;
+
+		/** Type of stream we are creating. Only 2110-20 (video) supported at the moment */
+		ERivermaxStreamType RivermaxStreamType = ERivermaxStreamType::VIDEO_2110_20_STREAM;
+
+		/** Information about the incoming frame we are receiving and general sequence numbering we are using for packet loss */
+		FInputStreamData StreamData;
+
+		/** Various stats kept for this stream */
+		FInputStreamStats StreamStats;
+
+		/** Creator / listener of this stream */
+		IRivermaxInputStreamListener* Listener;
+		
+		/** True when player has requested stream to shutdown */
+		std::atomic<bool> bIsShuttingDown = false;
+
+		/** Future used for the initialization task */
+		TFuture<void> InitTaskFuture;
+
+		/**  Timestamp used to know when to print out stats */
+		double LastLoggingTimestamp = 0.0;
+
+		/** Info about the video format that was selected such as pixel group, group size etc */
+		FVideoFormatInfo FormatInfo;
+		
+		/** Whether gpudirect is used for this stream. Even though listener requests it, doesn't mean it will/can be used */
+		bool bIsUsingGPUDirect = false;
+
+		/** Whether this stream has dynamic header data split enabled. Need to disable it on shutdown */
+		bool bIsDynamicHeaderEnabled = false;
+
+		/** GPU device index to use when mapping to cuda. Will be needed for mgpu support */
+		int32 GPUDeviceIndex = 0;
+
+		/** Allocated memory base address used when it's time to free */
+		void* GPUAllocatedMemoryBaseAddress = nullptr;
+
+		/** Total allocated gpu memory. */
+		int32 GPUAllocatedMemorySize = 0;
+
+		/** Map between buffer we are writing to and their mapped address in cuda space */
+		TMap<FRHIBuffer*, void*> BufferGPUMemoryMap;
+
+		/** Used to wait for a frame to be copied from gpu to gpu memory */
+		struct FCallbackPayload
+		{
+			volatile bool bIsWaitingForPendingCopy = false;
+		};
+		TSharedPtr<FCallbackPayload> CallbackPayload;
+
+		/** Cuda stream used for our operations */ 
+		void* GPUStream = nullptr;
+
+		/** Payload size used to initialize rivermax stream and know differences between received SRD and config */
+		uint32 ExpectedPayloadSize = 1500;
+
+		/** 
+		 * States used to track reception of a frame
+		 */
+		enum class EReceptionState
+		{
+			/** Used for the first frame receptionand when an error happened while receiving to detect reception of the end of frame */
+			WaitingForMarker, 
+
+			/** Used while receiving a frame and processing each SampleRowData */
+			Receiving,
+
+			/** Used when an error is detected during reception. Will wait for the end of frame detection before moving back to Receiving */
+			FrameError
+		};
+
+		/** Starts stream waiting for a marker to align ourselves with a full frame */
+		EReceptionState State = EReceptionState::WaitingForMarker;
+
+		/** Track incoming frame description to infer resolution, pixel type. */
+		FFrameDescriptionTrackingData FrameDescriptionTracking;
+
+		/** Resolution being used by this stream. Can either be auto detected or enforced by owner */
+		FIntPoint StreamResolution = FIntPoint::ZeroValue;
+
+		/* Pointer to the rivermax API to avoid virtual calls in a hot loop. */
+		const RIVERMAX_API_FUNCTION_LIST* CachedAPI = nullptr;
+	};
+}
+
+
diff --git a/Engine/Plugins/VirtualProduction/Rivermax/RivermaxCore/Source/RivermaxCore/Private/Streams/RivermaxOutAncStream.cpp b/Engine/Plugins/VirtualProduction/Rivermax/RivermaxCore/Source/RivermaxCore/Private/Streams/RivermaxOutAncStream.cpp
new file mode 100644
--- /dev/null
+++ b/Engine/Plugins/VirtualProduction/Rivermax/RivermaxCore/Source/RivermaxCore/Private/Streams/RivermaxOutAncStream.cpp
@@ -0,0 +1,235 @@
+// Copyright Epic Games, Inc. All Rights Reserved.
+
+#include "RivermaxOutAncStream.h"
+
+#include "Async/Async.h"
+#include "CudaModule.h"
+#include "IRivermaxCoreModule.h"
+#include "IRivermaxBoundaryMonitor.h"
+#include "IRivermaxManager.h"
+#include "Misc/ByteSwap.h"
+#include "RivermaxFrameAllocator.h"
+#include "RivermaxFrameManager.h"
+#include "RivermaxLog.h"
+#include "RivermaxPTPUtils.h"
+#include "RivermaxTracingUtils.h"
+#include "RivermaxTypes.h"
+#include "RivermaxUtils.h"
+
+
+
+#define GET_FRAME_INDEX(OUTPUT_FRAME) OUTPUT_FRAME->GetFrameCounter() % Options.NumberOfBuffers
+namespace UE::RivermaxCore::Private
+{
+	// Fills the Rivermax memory with ANC header and data.
+	void FillChunk(const TSharedPtr<FRivermaxAncOutputOptions> InStreamOptions, const FRivermaxOutputStreamMemory& InStreamMemory, uint8* InFirstPacketStartPointer, uint16* InPayloadSizes, const UE::RivermaxCore::Private::FRivermaxOutputStreamData& StreamData, TSharedPtr<UE::RivermaxCore::Private::FRivermaxOutputFrame> CurrentFrame)
+	{
+		constexpr uint8 PayloadType = 97;
+
+		// Anc field values.
+		constexpr uint8 ProgressiveField = 0b00;
+		constexpr uint8 InterlaceField = 0b10;
+		constexpr uint8 InterlaceSecondField = 0b11;
+
+		for (size_t PayloadIndex = 0; PayloadIndex < InStreamMemory.PacketsPerChunk; ++PayloadIndex, InFirstPacketStartPointer += InStreamMemory.PacketsPerChunk) {
+			memset(InFirstPacketStartPointer, 0, InStreamMemory.PayloadSize);
+			FAncRTPHeader* AncRtpHeader = (struct FAncRTPHeader*)&InFirstPacketStartPointer[0];
+			//RTP header initialization
+			AncRtpHeader->RTPHeader.Version = 2;
+			AncRtpHeader->RTPHeader.ExtensionBit = 0;
+			AncRtpHeader->RTPHeader.PaddingBit = 0;
+			AncRtpHeader->RTPHeader.MarkerBit = 1;
+			AncRtpHeader->RTPHeader.PayloadType = PayloadType;
+			AncRtpHeader->RTPHeader.SequenceNumber = ByteSwap((uint16)(StreamData.SequenceNumber & 0xFFFF));
+			AncRtpHeader->RTPHeader.Timestamp = ByteSwap(CurrentFrame->MediaTimestamp);
+			AncRtpHeader->RTPHeader.SynchronizationSource = StreamData.SynchronizationSource;
+			AncRtpHeader->RTPHeader.ExtendedSequenceNumber = ByteSwap((uint16)((StreamData.SequenceNumber >> 16) & 0xFFFF));
+
+			AncRtpHeader->ANCCount = 1;
+			AncRtpHeader->Field = ProgressiveField;
+
+			FAncillaryTimecodeHeaderFields* AncDataHeader = (struct FAncillaryTimecodeHeaderFields*)&InFirstPacketStartPointer[sizeof(FAncRTPHeader)];
+			AncDataHeader->DataPacketHeaderFields.ColorDiff = 0;
+			AncDataHeader->DataPacketHeaderFields.SetLineNumber(0x7FF);
+			AncDataHeader->DataPacketHeaderFields.SetHorizontalOffset(0xFFF);
+			AncDataHeader->DataPacketHeaderFields.StreamFlag = 0;
+			AncDataHeader->DataPacketHeaderFields.StreamNum = 0;
+			AncDataHeader->SetDID(InStreamOptions->DID);
+			AncDataHeader->SetSDID(InStreamOptions->SDID);
+			uint16 Length = sizeof(FAncillaryTimecodeHeaderFields);
+			AncDataHeader->SetATCTimecode(CurrentFrame->Timecode.Hours, CurrentFrame->Timecode.Minutes, CurrentFrame->Timecode.Seconds, CurrentFrame->Timecode.Frames, CurrentFrame->Timecode.bDropFrameFormat);
+			AncRtpHeader->Length = ByteSwap((uint16)(Length));
+			InPayloadSizes[PayloadIndex] = sizeof(FAncRTPHeader) /*RTP + ERTP headers*/ + Length;
+		}
+	}
+
+
+	FRivermaxOutAncStream::FRivermaxOutAncStream(const TArray<char>& SDPDescription)
+		: FRivermaxOutStream(SDPDescription)
+		, FrameInfoToSend(MakeShared<FRivermaxOutputAncInfo>())
+	{
+		StreamType = ERivermaxStreamType::ANC_2110_40_STREAM;
+	}
+
+	bool FRivermaxOutAncStream::PushFrame(TSharedPtr<IRivermaxOutputInfo> FrameInfo)
+	{
+		TRACE_CPUPROFILER_EVENT_SCOPE(FRivermaxOutAncStream::PushVideoFrame);
+
+		// Clear reserved frame if there is one. If not, try to get free frame.
+		// ReservedFrame should always be valid when block on reservation mode is used.
+		TSharedPtr<FRivermaxOutputFrame> ReservedFrame;
+		if (ReservedFrames.Dequeue(ReservedFrame))
+		{
+			// Sanity check.
+			//check(FrameInfo->FrameIdentifier == ReservedFrame->GetFrameCounter());
+		}
+		else
+		{
+			check(Options.FrameLockingMode != EFrameLockingMode::BlockOnReservation);
+		}
+
+		// If this is invalid it means that frame locking mode is BlockOnReservation and the render ran faster than media output fps.
+		FrameInfoToSend->FrameIdentifier = FrameInfo->FrameIdentifier;
+		FrameInfoToSend->Height = FrameInfo->Height;
+		FrameInfoToSend->Width = FrameInfo->Width;
+		FrameInfoToSend->Stride = FrameInfo->Stride;
+		FrameReadyToSendSignal->Trigger();
+		return true;
+	}
+
+	bool FRivermaxOutAncStream::InitializeStreamMemoryConfig()
+	{
+		using namespace UE::RivermaxCore::Private::Utils;
+
+		constexpr int64 PacketsPerFrame = 1;
+
+		// We might need a smaller packet to complete the end of frame so ceil to the next value
+		StreamMemory.PacketsPerFrame = PacketsPerFrame;
+
+		constexpr size_t StridesInChunk = 1;
+		StreamMemory.PacketsPerChunk = StridesInChunk;
+		StreamMemory.FramesFieldPerMemoryBlock = CachedCVars.bUseSingleMemblock ? Options.NumberOfBuffers : 1;
+
+		// Only one chunk with one packet.
+		constexpr int64 ChunksPerFrame = 1.;
+		StreamMemory.ChunksPerFrameField = ChunksPerFrame;
+		const uint64 RealPacketsPerFrame = StreamMemory.ChunksPerFrameField * StreamMemory.PacketsPerChunk;
+		StreamMemory.PacketsPerMemoryBlock = RealPacketsPerFrame * StreamMemory.FramesFieldPerMemoryBlock;
+		StreamMemory.ChunksPerMemoryBlock = StreamMemory.FramesFieldPerMemoryBlock * StreamMemory.ChunksPerFrameField;
+		StreamMemory.MemoryBlockCount = Options.NumberOfBuffers / StreamMemory.FramesFieldPerMemoryBlock;
+		StreamMemory.DataBlockID = 0;
+		// A size close to max UDP packet size. Since it is only one packet.
+		StreamMemory.PayloadSize = 1300;
+
+		StreamMemory.ChunkSpacingBetweenMemcopies = 1;
+		StreamMemory.bUseIntermediateBuffer = true;
+
+		if (!SetupFrameManagement())
+		{
+			return false;
+		}
+
+		StreamMemory.MemoryBlocks.SetNumZeroed(StreamMemory.MemoryBlockCount);
+		CachedAPI->rmx_output_media_init_mem_blocks(StreamMemory.MemoryBlocks.GetData(), StreamMemory.MemoryBlockCount);
+		for (uint32 BlockIndex = 0; BlockIndex < StreamMemory.MemoryBlockCount; ++BlockIndex)
+		{
+			rmx_output_media_mem_block& Block = StreamMemory.MemoryBlocks[BlockIndex];
+			CachedAPI->rmx_output_media_set_chunk_count(&Block, StreamMemory.ChunksPerMemoryBlock);
+
+			// Anc only needs one sub block since we don't split header and data.
+			constexpr uint8 SubBlockCount = 1;
+			CachedAPI->rmx_output_media_set_sub_block_count(&Block, SubBlockCount);
+		}
+		
+		return true;
+	}
+
+	bool FRivermaxOutAncStream::IsFrameAvailableToSend()
+	{
+		// This should also depend on Video stream.
+		return true;
+	}
+
+	bool FRivermaxOutAncStream::CopyFrameData(const TSharedPtr<FRivermaxOutputFrame>& SourceFrame, uint8* DestinationBase)
+	{
+		TSharedPtr<FRivermaxAncOutputOptions> StreamOptions = StaticCastSharedPtr<FRivermaxAncOutputOptions>(Options.StreamOptions[(uint8)StreamType]);
+
+		// Alternative to rmx_output_media_set_packet_layout. Required for dynamically sized packets.
+		CachedAPI->rmx_output_media_set_chunk_packet_count(&StreamData.ChunkHandle, (size_t)StreamMemory.PacketsPerChunk);
+
+		// This size will be filled with the actual sizes of the packets.
+		uint16* PayloadSizes = rmx_output_media_get_chunk_packet_sizes(&StreamData.ChunkHandle, StreamMemory.DataBlockID);
+
+		FillChunk(StreamOptions, StreamMemory, reinterpret_cast<uint8*>(CurrentFrame->FrameStartPtr), PayloadSizes, StreamData, CurrentFrame);
+
+		OnFrameReadyToBeSent();
+
+		return true;
+	}
+
+	bool FRivermaxOutAncStream::SetupFrameManagement()
+	{
+		TRACE_CPUPROFILER_EVENT_SCOPE(FRivermaxOutAncStream::SetupFrameManagement);
+
+		return true;
+	}
+
+	void FRivermaxOutAncStream::CleanupFrameManagement()
+	{
+	}
+
+	bool FRivermaxOutAncStream::ReserveFrame(uint64 FrameCounter) const
+	{
+		// There is only one reserved frame at the time per stream.
+		TSharedPtr<FRivermaxOutputFrame> ReservedFrame = MakeShared<FRivermaxOutputFrame>();
+
+		if (ReservedFrame.IsValid())
+		{
+			ReservedFrame->SetFrameCounter(FrameCounter);
+			ReservedFrames.Enqueue(ReservedFrame);
+		}
+
+		return ReservedFrame.IsValid();
+	}
+
+	TSharedPtr<FRivermaxOutputFrame> FRivermaxOutAncStream::GetNextFrameToSend(bool bWait)
+	{
+		TSharedPtr<FRivermaxOutputFrame> NextFrameToSend = MakeShared<FRivermaxOutputFrame>();
+
+		if (bWait)
+		{
+			while (!NextFrameToSend.IsValid() && bIsActive)
+			{
+				FrameReadyToSendSignal->Wait();
+				//NextFrameToSend = FrameManager->DequeueFrameToSend();
+			}
+		}
+		return NextFrameToSend;
+	}
+
+	void FRivermaxOutAncStream::LogStreamDescriptionOnCreation() const
+	{
+		FRivermaxOutStream::LogStreamDescriptionOnCreation();
+
+		TStringBuilder<512> StreamDescription;
+
+		TSharedPtr<FRivermaxAncOutputOptions> StreamOptions = StaticCastSharedPtr<FRivermaxAncOutputOptions>(Options.StreamOptions[(uint8)StreamType]);
+		StreamDescription.Appendf(TEXT("FrameRate = %s, "), *StreamOptions->FrameRate.ToPrettyText().ToString());
+		StreamDescription.Appendf(TEXT("Alignment = %s, "), LexToString(Options.AlignmentMode));
+		StreamDescription.Appendf(TEXT("Framelocking = %s."), LexToString(Options.FrameLockingMode));
+
+		UE_LOG(LogRivermax, Display, TEXT("%s"), *FString(StreamDescription));
+	}
+
+	void FRivermaxOutAncStream::SetupRTPHeaders()
+	{
+		// ANC RTP headers are shipped with the payload. All needs to be done is incrementation of packets and sequence number.
+		++StreamData.SequenceNumber;
+	}
+
+	void FRivermaxOutAncStream::CompleteCurrentFrame(bool bReleaseFrame)
+	{
+		FRivermaxOutStream::CompleteCurrentFrame(bReleaseFrame);
+	}
+}
+
diff --git a/Engine/Plugins/VirtualProduction/Rivermax/RivermaxCore/Source/RivermaxCore/Private/Streams/RivermaxOutAncStream.h b/Engine/Plugins/VirtualProduction/Rivermax/RivermaxCore/Source/RivermaxCore/Private/Streams/RivermaxOutAncStream.h
new file mode 100644
--- /dev/null
+++ b/Engine/Plugins/VirtualProduction/Rivermax/RivermaxCore/Source/RivermaxCore/Private/Streams/RivermaxOutAncStream.h
@@ -0,0 +1,56 @@
+// Copyright Epic Games, Inc. All Rights Reserved.
+
+#pragma once
+
+#include "RivermaxOutStream.h"
+
+#include "Async/Future.h"
+#include "Containers/SpscQueue.h"
+#include "HAL/Runnable.h"
+#include "RivermaxWrapper.h"
+#include "RivermaxOutputFrame.h"
+#include "RivermaxTypes.h"
+#include "RTPHeader.h"
+
+
+class FEvent;
+class IRivermaxCoreModule;
+
+namespace UE::RivermaxCore::Private
+{
+	using UE::RivermaxCore::FRivermaxOutputOptions;
+
+	class FRivermaxOutAncStream : public FRivermaxOutStream
+	{
+	public:
+		FRivermaxOutAncStream(const TArray<char>& SDPDescription);
+
+	public:
+
+		//~ Begin IRivermaxOutputStream interface
+		virtual bool PushFrame(TSharedPtr<IRivermaxOutputInfo> FrameInfo) override;
+		virtual bool ReserveFrame(uint64 FrameCounter) const override;
+		//~ End IRivermaxOutputStream interface
+
+	protected:
+
+		//~ Begin FRivermaxOutStream interface. Refer to the parent class for more information.
+		virtual bool IsFrameAvailableToSend();
+		virtual bool InitializeStreamMemoryConfig();
+		virtual bool CopyFrameData(const TSharedPtr<FRivermaxOutputFrame>& SourceFrame, uint8* DestinationBase);
+		virtual bool SetupFrameManagement();
+		virtual void CleanupFrameManagement();
+		virtual void CompleteCurrentFrame(bool bReleaseFrame) override;
+		virtual void SetupRTPHeaders() override;
+		virtual void LogStreamDescriptionOnCreation() const override;
+		virtual TSharedPtr<FRivermaxOutputFrame> GetNextFrameToSend(bool bWait = false);
+		//~ End FRivermaxOutStream interface
+
+	private:
+
+		/** All the information required for ANC data to be sent. */
+		TSharedPtr<FRivermaxOutputAncInfo> FrameInfoToSend;
+	};
+}
+
+
diff --git a/Engine/Plugins/VirtualProduction/Rivermax/RivermaxCore/Source/RivermaxCore/Private/Streams/RivermaxOutStream.cpp b/Engine/Plugins/VirtualProduction/Rivermax/RivermaxCore/Source/RivermaxCore/Private/Streams/RivermaxOutStream.cpp
new file mode 100644
--- /dev/null
+++ b/Engine/Plugins/VirtualProduction/Rivermax/RivermaxCore/Source/RivermaxCore/Private/Streams/RivermaxOutStream.cpp
@@ -0,0 +1,1222 @@
+// Copyright Epic Games, Inc. All Rights Reserved.
+
+#include "RivermaxOutStream.h"
+
+#include "Async/Async.h"
+#include "CudaModule.h"
+#include "IRivermaxCoreModule.h"
+#include "IRivermaxBoundaryMonitor.h"
+#include "IRivermaxManager.h"
+#include "Misc/ByteSwap.h"
+#include "RivermaxFrameAllocator.h"
+#include "RivermaxFrameManager.h"
+#include "RivermaxLog.h"
+#include "RivermaxPTPUtils.h"
+#include "RivermaxTracingUtils.h"
+#include "RivermaxTypes.h"
+#include "RivermaxUtils.h"
+
+#define GET_FRAME_INDEX(OUTPUT_FRAME) OUTPUT_FRAME->GetFrameCounter() % Options.NumberOfBuffers
+namespace UE::RivermaxCore::Private
+{
+	static TAutoConsoleVariable<int32> CVarRivermaxWakeupOffset(
+		TEXT("Rivermax.WakeupOffset"), 0,
+		TEXT("Wakeup is done on alignment point. This offset will be substracted from it to wake up earlier. Units: nanoseconds"),
+		ECVF_Default);
+
+	static TAutoConsoleVariable<int32> CVarRivermaxScheduleOffset(
+		TEXT("Rivermax.ScheduleOffset"), 0,
+		TEXT("Scheduling is done at alignment point plus TRO. This offset will be added to it to delay or schedule earlier. Units: nanoseconds"),
+		ECVF_Default);
+
+	static TAutoConsoleVariable<int32> CVarRivermaxOutputShowStats(
+		TEXT("Rivermax.ShowOutputStats"), 0,
+		TEXT("Enable stats logging at fixed interval"),
+		ECVF_Default);
+
+	static TAutoConsoleVariable<float> CVarRivermaxOutputShowStatsInterval(
+		TEXT("Rivermax.ShowStatsInterval"), 1.0,
+		TEXT("Interval at which to show stats in seconds"),
+		ECVF_Default);
+
+	static TAutoConsoleVariable<float> CVarRivermaxOutputTROOverride(
+		TEXT("Rivermax.Output.TRO"), 0,
+		TEXT("If not 0, overrides transmit offset calculation (TRO) based on  frame rate and resolution with a fixed value. Value in seconds."),
+		ECVF_Default);
+
+	static TAutoConsoleVariable<int32> CVarRivermaxOutputSkipSchedulingCutOffTime(
+		TEXT("Rivermax.Output.Scheduling.SkipCutoff"), 50,
+		TEXT("Required time in microseconds from scheduling time to avoid skipping an interval."),
+		ECVF_Default);
+
+	static TAutoConsoleVariable<int32> CVarRivermaxOutputForceImmediateSchedulingThreshold(
+		TEXT("Rivermax.Output.Scheduling.ForceImmediateCutoff"), 600,
+		TEXT("Required time in nanoseconds from scheduling time before we clamp to do it immediately."),
+		ECVF_Default);
+
+	static TAutoConsoleVariable<int32> CVarRivermaxOutputUseSingleMemblock(
+		TEXT("Rivermax.Output.UseSingleMemblock"), 1,
+		TEXT("Configures Rivermax stream to use a single memblock potentially improving SDK performance."),
+		ECVF_Default);
+
+	static TAutoConsoleVariable<int32> CVarRivermaxOutputRandomDelay(
+		TEXT("Rivermax.Output.TriggerRandomDelay"), 0,
+		TEXT("Will cause a delay of variable amount of time when next frame is sent."),
+		ECVF_Default);
+
+	static TAutoConsoleVariable<int32> CVarRivermaxOutputEnableTimingProtection(
+		TEXT("Rivermax.Output.Scheduling.EnableTimingProtection"), 1,
+		TEXT("Whether timing verification is done on commit to avoid misalignment. Next frame interval is skipped if it happens."),
+		ECVF_Default);
+
+	static TAutoConsoleVariable<int32> CVarRivermaxOutputEnableTimeCriticalThread(
+		TEXT("Rivermax.Output.EnableTimeCriticalThread"), 0,
+		TEXT("Whether to set output thread as time critical."),
+		ECVF_Default);
+
+	static TAutoConsoleVariable<bool> CVarRivermaxOutputPrefillRTPHeaders(
+		TEXT("Rivermax.Output.PrefillRTPHeaders"), true,
+		TEXT("Optimization used to prefill every RTP headers with known data."),
+		ECVF_Default);
+
+	static TAutoConsoleVariable<float> CVarRivermaxOutputCommitChunksOffsetPercent(
+		TEXT("Rivermax.Output.CommitChunksOffsetPercent"), 0.3,
+		TEXT("This CVar will allow for chunks to be committed before the next alignment point if it is at all possible.\n\
+		The value indicates in percent of frame time how much earlier the Rivermax plugin will attempt to commit chunks."),
+		ECVF_Default);
+
+	static bool GbTriggerRandomTimingIssue = false;
+	FAutoConsoleVariableRef CVarTriggerRandomTimingIssue(
+		TEXT("Rivermax.Sync.TriggerRandomTimingIssue")
+		, UE::RivermaxCore::Private::GbTriggerRandomTimingIssue
+		, TEXT("Randomly triggers a timing issue to test self repair."), ECVF_Cheat);
+
+	FRivermaxOutStream::FRivermaxOutStream(const TArray<char>& InSDPDescription)
+		: bIsActive(false)
+		, SDPDescription(InSDPDescription)
+	{
+
+	}
+
+	FRivermaxOutStream::~FRivermaxOutStream()
+	{
+		Uninitialize();
+	}
+
+	bool FRivermaxOutStream::Initialize(const FRivermaxOutputOptions& InOptions, IRivermaxOutputStreamListener& InListener)
+	{
+		TRACE_CPUPROFILER_EVENT_SCOPE(FRivermaxOutStream::Initialize);
+
+		RivermaxModule = FModuleManager::GetModulePtr<IRivermaxCoreModule>(TEXT("RivermaxCore"));
+		if (RivermaxModule->GetRivermaxManager()->ValidateLibraryIsLoaded() == false)
+		{
+			UE_LOG(LogRivermax, Warning, TEXT("Can't create Rivermax Output Stream. Library isn't initialized."));
+			return false;
+		}
+
+		Options = InOptions;
+		Listener = &InListener;
+
+		if (IConsoleVariable* CvarDelay = CVarRivermaxOutputRandomDelay.AsVariable())
+		{
+			CvarDelay->OnChangedDelegate().AddRaw(this, &FRivermaxOutStream::OnCVarRandomDelayChanged);
+		}
+
+		CacheCVarValues();
+		
+		// Cache API entry point
+		CachedAPI = RivermaxModule->GetRivermaxManager()->GetApi();
+		checkSlow(CachedAPI);
+
+		// Enable frame boundary monitoring
+		MonitoringGuid = RivermaxModule->GetRivermaxBoundaryMonitor().StartMonitoring(GetFrameRate());
+
+		if (!InitializeStreamMemoryConfig())
+		{
+			return false;
+		}
+
+		InitializationFuture = Async(EAsyncExecution::TaskGraph, [this]()
+		{
+			// Create Rivermax stream using memory configuration
+			{
+				// Setup audio stream settings
+				rmx_output_media_stream_params OutputStreamParameters;
+				CachedAPI->rmx_output_media_init(&OutputStreamParameters);
+				CachedAPI->rmx_output_media_set_sdp(&OutputStreamParameters, &SDPDescription[0]);
+				CachedAPI->rmx_output_media_assign_mem_blocks(&OutputStreamParameters, StreamMemory.MemoryBlocks.GetData(), StreamMemory.MemoryBlocks.Num());
+
+				// Priority Code Point for Quality of Service. Tells the network how important this packet is. 7 is highest priority
+				constexpr uint8 PCPAttribute = 7;
+				CachedAPI->rmx_output_media_set_pcp(&OutputStreamParameters, PCPAttribute);
+
+				if (StreamType == ERivermaxStreamType::AUDIO_2110_30_STREAM)
+				{
+					//Differentiated Services Code Point. For AES67 RTP media streams, the DSCP value is set to 34.
+					constexpr uint8 DSCP = 34;
+					CachedAPI->rmx_output_media_set_dscp(&OutputStreamParameters, DSCP);
+				}
+
+				constexpr uint8 ECN = 0; //Explicit congestion notification in theory notifies if the packet is "congested"
+				CachedAPI->rmx_output_media_set_ecn(&OutputStreamParameters, ECN);
+
+				// Sometimes, chunk count will have more packets than needed so last ones might be 0 sized. 
+				// Verify if new API work with the actual amount of packet with data or it needs the padded version 
+				CachedAPI->rmx_output_media_set_packets_per_frame(&OutputStreamParameters, StreamMemory.ChunksPerFrameField * StreamMemory.PacketsPerChunk);
+				CachedAPI->rmx_output_media_set_packets_per_chunk(&OutputStreamParameters, StreamMemory.PacketsPerChunk);
+
+				// This means that this stream doesn't need dynamic header split. Refer to the description of HeaderBlockID and DataBlockID
+				if (StreamMemory.DataBlockID != 0)
+				{
+					CachedAPI->rmx_output_media_set_stride_size(&OutputStreamParameters, StreamMemory.HeaderBlockID, StreamMemory.HeaderStrideSize);
+				}
+				CachedAPI->rmx_output_media_set_stride_size(&OutputStreamParameters, StreamMemory.DataBlockID, StreamMemory.PayloadSize);
+
+				const size_t MediaBlockIndex = GetStreamIndexSDP();
+				
+				CachedAPI->rmx_output_media_set_idx_in_sdp(&OutputStreamParameters, MediaBlockIndex);
+
+				rmx_stream_id NewId;
+				rmx_status Status = CachedAPI->rmx_output_media_create_stream(&OutputStreamParameters, &NewId);
+
+				if (Status == RMX_OK)
+				{
+					struct sockaddr_in SourceAddress;
+					FMemory::Memset(&SourceAddress, 0, sizeof(SourceAddress));
+
+					rmx_output_media_context MediaContext;
+					CachedAPI->rmx_output_media_init_context(&MediaContext, NewId);
+
+					const size_t SDPMediaIndex = MediaBlockIndex;
+					CachedAPI->rmx_output_media_set_context_block(&MediaContext, SDPMediaIndex);
+					Status = CachedAPI->rmx_output_media_get_local_address(&MediaContext, reinterpret_cast<sockaddr*>(&SourceAddress));
+					if (Status == RMX_OK)
+					{
+						struct sockaddr_in DestinationAddress;
+						FMemory::Memset(&DestinationAddress, 0, sizeof(DestinationAddress));
+
+						Status = CachedAPI->rmx_output_media_get_remote_address(&MediaContext, reinterpret_cast<sockaddr*>(&DestinationAddress));
+						if (Status == RMX_OK)
+						{
+							StreamId = NewId;
+
+							CachedAPI->rmx_output_media_init_chunk_handle(&StreamData.ChunkHandle, StreamId);
+
+							// This should be unique for each stream. Will be used in packet creation.
+							StreamData.SynchronizationSource = FRawRTPHeader::VideoSynchronizationSource + (uint32)MediaBlockIndex;
+
+							StreamData.FrameFieldTimeIntervalNs = 1E9 / GetFrameRate().AsDecimal();
+							InitializeStreamTimingSettings();
+
+							LogStreamDescriptionOnCreation();
+							FString SDPAsAString(ANSI_TO_TCHAR(&SDPDescription[0]));
+							UE_LOG(LogRivermax, Verbose, TEXT("Created stream using SDP:\n%s"), *SDPAsAString);
+
+							bIsActive = true;
+							RivermaxThread.Reset(FRunnableThread::Create(this, TEXT("Rmax OutputStream Thread"), 128 * 1024, TPri_TimeCritical, FPlatformAffinity::GetPoolThreadMask()));
+						}
+						else
+						{
+							UE_LOG(LogRivermax, Warning, TEXT("Failed querying destination address. Output Stream won't be created. Status: %d"), Status);
+						}
+					}
+					else
+					{
+						UE_LOG(LogRivermax, Warning, TEXT("Failed querying local address. Output Stream won't be created. Status: %d"), Status);
+					}
+				}
+				else
+				{
+					UE_LOG(LogRivermax, Warning, TEXT("Failed to create Rivermax output stream. Status: %d"), Status);
+				}
+			}
+
+			Listener->OnInitializationCompleted(bIsActive);
+		});
+
+		return true;
+	}
+
+	void FRivermaxOutStream::Uninitialize()
+	{
+		if (!InitializationFuture.IsReady())
+		{
+			InitializationFuture.Wait();
+		}
+
+		if (RivermaxThread != nullptr)
+		{
+			Stop();
+			
+			FrameAvailableSignal->Trigger();
+			FrameReadyToSendSignal->Trigger();
+			RivermaxThread->Kill(true);
+			RivermaxThread.Reset();
+
+			CleanupFrameManagement();
+
+			RivermaxModule->GetRivermaxBoundaryMonitor().StopMonitoring(MonitoringGuid, GetFrameRate());
+			
+			UE_LOG(LogRivermax, Log, TEXT("Rivermax Output stream has shutdown"));
+		}
+
+		if (IConsoleVariable* CvarDelay = CVarRivermaxOutputRandomDelay.AsVariable())
+		{
+			CvarDelay->OnChangedDelegate().RemoveAll(this);
+		}
+	}
+
+	void FRivermaxOutStream::InitializeNextFrame(const TSharedPtr<FRivermaxOutputFrame>& NextFrame)
+	{
+		NextFrame->LineNumber = 0;
+		NextFrame->PacketCounter = 0;
+		NextFrame->ChunkNumber = 0;
+		NextFrame->PayloadPtr = nullptr;
+		NextFrame->HeaderPtr = nullptr;
+		NextFrame->FrameStartPtr = nullptr;
+		NextFrame->bCaughtTimingIssue = false;
+		NextFrame->Offset = 0;
+	}
+	
+	void FRivermaxOutStream::DestroyStream()
+	{
+		rmx_status Status = CachedAPI->rmx_output_media_cancel_unsent_chunks(&StreamData.ChunkHandle);
+		if (Status != RMX_OK)
+		{
+			UE_LOG(LogRivermax, Warning, TEXT("Could not cancel unsent chunks when destroying output stream. Status: %d"), Status);
+		}
+
+		do 
+		{
+			Status = CachedAPI->rmx_output_media_destroy_stream(StreamId);
+			if (RMX_BUSY == Status) 
+			{
+				FPlatformProcess::SleepNoStats(0.3);
+			}
+		} while (Status == RMX_BUSY);
+	}
+
+	bool FRivermaxOutStream::WaitForNextRound()
+	{
+		TRACE_CPUPROFILER_EVENT_SCOPE(WaitForNextRound);
+		const uint64 CurrentTimeNanosec = RivermaxModule->GetRivermaxManager()->GetTime();
+		const double CurrentPlatformTime = FPlatformTime::Seconds();
+		const uint64 CurrentFrameNumber = UE::RivermaxCore::GetFrameNumber(CurrentTimeNanosec, GetFrameRate());
+
+		switch (Options.AlignmentMode)
+		{
+			case ERivermaxAlignmentMode::AlignmentPoint:
+			{
+				CalculateNextScheduleTime_AlignementPoints(CurrentTimeNanosec, CurrentFrameNumber);
+				break;
+			}
+			case ERivermaxAlignmentMode::FrameCreation:
+			{
+				CalculateNextScheduleTime_FrameCreation(CurrentTimeNanosec, CurrentFrameNumber);
+				break;
+			}
+			default:
+			{
+				checkNoEntry();
+			}
+		}
+
+		// Offset wakeup if desired to give more time for scheduling. 
+		const uint64 WakeupTime = StreamData.NextAlignmentPointNanosec - CVarRivermaxWakeupOffset.GetValueOnAnyThread();
+
+		uint64 WaitTimeNanosec = WakeupTime - CurrentTimeNanosec;
+
+		// Wakeup can be smaller than current time with controllable offset
+		if (WakeupTime < CurrentTimeNanosec)
+		{
+			WaitTimeNanosec = 0;
+		}
+
+		static constexpr float SleepThresholdSec = 5.0f / 1000;
+		static constexpr float YieldTimeSec = 2.0f / 1000;
+		const double WaitTimeSec = FMath::Min(WaitTimeNanosec / 1E9, 1.0);
+		StreamData.LastSleepTimeNanoSec = WaitTimeNanosec;
+
+		bool bIsFrameReady = false;
+		if (StreamMemory.bUseIntermediateBuffer)
+		{	
+			// When using intermediate buffer, we verify if next buffer is ready sooner than wake up time
+			// If a frame is ready already, we can move on. Otherwise, we wait for FrameReady signal with a
+			// wait timeout. In the case of a repeated frame, we will always timeout and we won't be able to 
+			// do an early copy.
+			if (IsFrameAvailableToSend())
+			{
+				bIsFrameReady = true;
+			}
+			else
+			{
+				const uint32 WaitMs = FMath::Floor((WaitTimeSec - YieldTimeSec) * 1000.);
+				do 
+				{
+					bIsFrameReady = FrameReadyToSendSignal->Wait(WaitMs);
+				} while (!IsFrameAvailableToSend() && bIsFrameReady && bIsActive);
+			}
+		}
+		else
+		{
+			// Sleep for the largest chunk of time
+			if (WaitTimeSec > SleepThresholdSec)
+			{
+				FPlatformProcess::SleepNoStats(WaitTimeSec - YieldTimeSec);
+			}
+		}
+
+		if (!bIsFrameReady)
+		{
+			// We are past the long sleep so no more early data access possible. Just yield until the wake up time.
+			{
+				// Use platform time instead of rivermax get PTP to avoid making calls to it. Haven't been profiled if it impacts
+				while (FPlatformTime::Seconds() < (CurrentPlatformTime + WaitTimeSec))
+				{
+					FPlatformProcess::SleepNoStats(0.f);
+				}
+			}
+
+			if (StreamData.bHasValidNextFrameNumber && CachedCVars.bShowOutputStats)
+			{
+				const uint64 AfterSleepTimeNanosec = RivermaxModule->GetRivermaxManager()->GetTime();
+				const uint64 RealWaitNs = AfterSleepTimeNanosec - CurrentTimeNanosec;
+				const uint64 OvershootSleep = AfterSleepTimeNanosec > StreamData.NextAlignmentPointNanosec ? AfterSleepTimeNanosec - StreamData.NextAlignmentPointNanosec : 0;
+				const double OvershootSleepSec = OvershootSleep / 1e9;
+
+				UE_LOG(LogRivermax, Verbose, TEXT("CurrentTime %llu. OvershootSleep: %0.9f. ExpectedWait: %0.9f. RealWait: %0.9f, Scheduling at %llu. NextAlign %llu. ")
+					, CurrentTimeNanosec
+					, OvershootSleepSec
+					, (double)WaitTimeNanosec / 1E9
+					, (double)RealWaitNs / 1E9
+					, StreamData.NextScheduleTimeNanosec
+					, StreamData.NextAlignmentPointNanosec);
+			}
+		}
+		else
+		{
+			if (StreamData.bHasValidNextFrameNumber && CachedCVars.bShowOutputStats)
+			{
+				const uint64 AfterSleepTimeNanosec = RivermaxModule->GetRivermaxManager()->GetTime();
+				UE_LOG(LogRivermax, Verbose, TEXT("Early data available. CurrentTime %llu. Scheduling at %llu. NextAlign %llu. ")
+					, AfterSleepTimeNanosec
+					, StreamData.NextScheduleTimeNanosec
+					, StreamData.NextAlignmentPointNanosec
+					, (StreamData.NextScheduleTimeNanosec - AfterSleepTimeNanosec));
+			}
+		}
+
+		return bIsFrameReady;
+	}
+
+	void FRivermaxOutStream::GetNextChunk()
+	{
+		TRACE_CPUPROFILER_EVENT_SCOPE(GetNextChunk);
+
+		bool bHasAddedTrace = false;
+		rmx_status Status;
+		do
+		{
+			Status = CachedAPI->rmx_output_media_get_next_chunk(&StreamData.ChunkHandle);
+			CurrentFrame->PayloadPtr = rmx_output_media_get_chunk_strides(&StreamData.ChunkHandle, StreamMemory.DataBlockID);
+
+			// This means that this stream doesn't need dynamic header split. Refer to the description of HeaderBlockID and DataBlockID
+			if (StreamMemory.DataBlockID != 0)
+			{
+				CurrentFrame->HeaderPtr = rmx_output_media_get_chunk_strides(&StreamData.ChunkHandle, StreamMemory.HeaderBlockID);
+			}
+
+			if (Status == RMX_OK)
+			{
+				if (CurrentFrame->FrameStartPtr == nullptr)
+				{
+					// Stamp frame start in order to copy frame data sequentially as we query chunks
+					CurrentFrame->FrameStartPtr = CurrentFrame->PayloadPtr;
+				}
+
+				break;
+			}
+			else if (Status == RMX_NO_FREE_CHUNK)
+			{
+				//We should not be here
+				if (!bHasAddedTrace)
+				{
+					Stats.LastFrameChunkRetries++;
+					UE_LOG(LogRivermax, Verbose, TEXT("No free chunks to get for chunk '%u'. Waiting"), CurrentFrame->ChunkNumber);
+					TRACE_CPUPROFILER_EVENT_SCOPE(GetNextChunk::NoFreeChunk);
+					bHasAddedTrace = true;
+				}
+			}
+			else
+			{
+				UE_LOG(LogRivermax, Error, TEXT("Invalid error happened while trying to get next chunks. Status: %d"), Status);
+				Listener->OnStreamError();
+				Stop();
+			}
+		} while (Status != RMX_OK && bIsActive);
+	}
+
+	void FRivermaxOutStream::CommitNextChunks()
+	{
+		TRACE_CPUPROFILER_EVENT_SCOPE(CommitNextChunks);
+		rmx_status Status;
+		int32 ErrorCount = 0;
+		const uint64 CurrentTimeNanosec = RivermaxModule->GetRivermaxManager()->GetTime();
+		uint64 ScheduleTime = CurrentFrame->ChunkNumber == 0 ? StreamData.NextScheduleTimeNanosec : 0;
+		if (CurrentFrame->ChunkNumber == 0)
+		{
+			check(true);
+		}
+		do
+		{
+			//Only first chunk gets scheduled with a timestamp. Following chunks are queued after it using 0
+			if (ScheduleTime != 0)
+			{
+				// If scheduling time is not far away enough, force it immediately otherwise rmax_commit will throw an error 
+				if (ScheduleTime <= (CurrentTimeNanosec + CachedCVars.ForceCommitImmediateTimeNanosec))
+				{
+					ScheduleTime = 0;
+					++Stats.CommitImmediate;
+				}
+			}
+
+			checkSlow(CachedAPI);
+
+			Status = CachedAPI->rmx_output_media_commit_chunk(&StreamData.ChunkHandle, ScheduleTime);
+
+			if (Status == RMX_OK)
+			{
+				break;
+			}
+			else if (Status == RMX_HW_SEND_QUEUE_IS_FULL)
+			{
+				Stats.CommitRetries++;
+				TRACE_CPUPROFILER_EVENT_SCOPE(CommitNextChunks::QUEUEFULL);
+				++ErrorCount;
+			}
+			else if (Status == RMX_HW_COMPLETION_ISSUE)
+			{
+				UE_LOG(LogRivermax, Error, TEXT("Completion issue while trying to commit next round of chunks."));
+				Listener->OnStreamError();
+				Stop();
+			}
+			else
+			{
+				UE_LOG(LogRivermax, Error, TEXT("Unhandled error (%d) while trying to commit next round of chunks."), Status);
+				Listener->OnStreamError();
+				Stop();
+			}
+
+		} while (Status != RMX_OK && bIsActive);
+
+		if (bIsActive && CurrentFrame->ChunkNumber == 0 && CachedCVars.bShowOutputStats)
+		{
+			UE_LOG(LogRivermax, Verbose, TEXT("Committed frame [%u]. Scheduled for '%llu'. Aligned with '%llu'. Current time '%llu'. Was late: %d. Slack: %llu. Errorcount: %d")
+				, CurrentFrame->GetFrameCounter()
+				, ScheduleTime
+				, StreamData.NextAlignmentPointNanosec
+				, CurrentTimeNanosec
+				, CurrentTimeNanosec >= StreamData.NextScheduleTimeNanosec ? 1 : 0
+				, StreamData.NextScheduleTimeNanosec >= CurrentTimeNanosec ? StreamData.NextScheduleTimeNanosec - CurrentTimeNanosec : 0
+				, ErrorCount);
+		}
+	}
+
+	bool FRivermaxOutStream::Init()
+	{
+		return true;
+	}
+
+	void FRivermaxOutStream::Process_AnyThread()
+	{
+		using namespace UE::RivermaxCore::Private::Utils;
+
+		// Wait for the next time a frame should be sent (based on frame interval)
+		// if interm buffer is used (alignment points) and a frame is ready before frame interval
+		//		Start copying data into next memory block from the intermediate buffer
+		//		At frame interval:
+		//			Release last sent frame if any
+		//			Make next frame the one being sent
+		// Otherwise
+		//		FrameCreation:
+		//			Release last sent frame if any
+		//			Wait for a new frame to be available
+		//		Alignment points:
+		//			Release last sent frame if any
+		// 
+		// Send frame
+		//		Get next chunk
+		//		Continue copy to intermediate buffer if required
+		//		Fill dynamic data for RTP headers of next chunk
+		//		Commit next chunk
+		// 
+		// Restart
+		{
+			bool bCanEarlyCopy = false;
+			{
+				TRACE_CPUPROFILER_EVENT_SCOPE(RmaxOut::Wait);
+				bCanEarlyCopy = WaitForNextRound();
+			}
+
+			if (bIsActive && bCanEarlyCopy)
+			{
+				TRACE_CPUPROFILER_EVENT_SCOPE(RmaxOut::PreprocessNextFrame);
+				PreprocessNextFrame();
+			}
+			else
+			{
+				TRACE_CPUPROFILER_EVENT_SCOPE(RmaxOut::PrepareNextFrame);
+				PrepareNextFrame();
+			}
+
+			// At this point, if there is no frame to send, move on to wait for next round
+			if (CurrentFrame.IsValid() && bIsActive)
+			{
+				SendFrame();
+
+				// If frame that was just sent failed timing requirements, we have to tell rivermax to skip 0 chunks 
+				// in order to reset internal states. Otherwise, scheduling time / Tro isn't respected next time we schedule.
+				if (CurrentFrame->bCaughtTimingIssue)
+				{
+					++Stats.TimingIssueCount;
+					constexpr uint64 ChunksToSkip = 0;
+					SkipChunks(ChunksToSkip);
+				}
+			}
+
+			Stats.TotalChunkRetries += Stats.LastFrameChunkRetries;
+			Stats.LastFrameChunkRetries = 0;
+		}
+	}
+
+	void FRivermaxOutStream::PreprocessNextFrame()
+	{
+		checkSlow(Options.AlignmentMode == ERivermaxAlignmentMode::AlignmentPoint);
+
+		TSharedPtr<FRivermaxOutputFrame> NextFrameToSend = GetNextFrameToSend();
+		if (StreamType == ERivermaxStreamType::ANC_2110_40_STREAM)
+		{
+			check(true);
+		}
+		if (ensure(NextFrameToSend))
+		{
+			InitializeNextFrame(NextFrameToSend);
+
+			// Now that we have the next frame, we can start copying data into it.
+			// We can't get chunks since commit will only commit the chunks returned by last call to get next chunk. 
+			// So, we calculate the next data and header pointer based on the current frame. 
+			const uint64 CurrentRmaxTimeNanosec = RivermaxModule->GetRivermaxManager()->GetTime();
+			const double CurrentPlatformTime = FPlatformTime::Seconds();
+			const double TimeLeftSec = double(StreamData.NextAlignmentPointNanosec - CurrentRmaxTimeNanosec) / 1E9;
+			const double TargetPlatformTimeSec = CurrentPlatformTime + TimeLeftSec;
+			if (CurrentRmaxTimeNanosec < StreamData.NextAlignmentPointNanosec)
+			{
+				if (!StreamMemory.bUseIntermediateBuffer)
+				{
+					TRACE_CPUPROFILER_EVENT_SCOPE(RmaxOut::CopyFrame);
+
+					bool bHasDataToCopy = true;
+					while (FPlatformTime::Seconds() < TargetPlatformTimeSec && bHasDataToCopy && bIsActive)
+					{
+						bHasDataToCopy = CopyFrameData(NextFrameToSend, reinterpret_cast<uint8*>(StreamMemory.BufferAddresses[GET_FRAME_INDEX(NextFrameToSend)]));
+					}
+				}
+
+				const double PostCopyTimeLeftSec = TargetPlatformTimeSec - FPlatformTime::Seconds();
+				if (PostCopyTimeLeftSec > 0)
+				{
+					TRACE_CPUPROFILER_EVENT_SCOPE(RmaxOut::Waiting);
+					static constexpr float YieldTimeSec = 1.0f / 1000;
+
+					double TargetPlatformTimeToCommitSec = TargetPlatformTimeSec;
+
+					const double CVarStartCommitEarlyPercent = FMath::Clamp(CVarRivermaxOutputCommitChunksOffsetPercent.GetValueOnAnyThread(), 0.0, 0.8);
+					if (CVarStartCommitEarlyPercent > 0.)
+					{
+						// We would like to start comitting chunks before the alignment point, but not too early so that we don't block this thread.
+						TargetPlatformTimeToCommitSec = TargetPlatformTimeSec - GetFrameRate().AsInterval() * CVarStartCommitEarlyPercent;
+					}
+
+					while (FPlatformTime::Seconds() < TargetPlatformTimeToCommitSec && bIsActive)
+					{
+						const double TimeLeft = TargetPlatformTimeToCommitSec - FPlatformTime::Seconds();
+						const double SleepTime = TimeLeft > YieldTimeSec ? TimeLeft - YieldTimeSec : 0.0;
+						FPlatformProcess::SleepNoStats(SleepTime);
+					}
+				}
+			}
+
+			{
+				TRACE_CPUPROFILER_EVENT_SCOPE(RmaxOut::WrappingUp);
+				if (CurrentFrame)
+				{
+					constexpr bool bReleaseFrame = true;
+					CompleteCurrentFrame(bReleaseFrame);
+				}
+
+				// Make the next frame to send the current one and update its state
+				CurrentFrame = MoveTemp(NextFrameToSend);
+			}
+		}
+		else
+		{
+			UE_LOG(LogRivermax, Error, TEXT("Unexpected error, no frame was available."));
+			Listener->OnStreamError();
+			Stop();
+		}
+	}
+
+	uint32 FRivermaxOutStream::Run()
+	{
+		if (CachedCVars.bEnableTimeCriticalThread)
+		{
+			::SetThreadPriority(GetCurrentThread(), THREAD_PRIORITY_TIME_CRITICAL);
+		}
+
+		// Initial wait for a frame to be produced
+		{
+			TRACE_CPUPROFILER_EVENT_SCOPE(Rmax::InitialWait);
+			FrameReadyToSendSignal->Wait();
+		}
+
+		while (bIsActive)
+		{
+			ShowStats();
+			Process_AnyThread();
+		}
+
+		DestroyStream();
+
+		return 0;
+	}
+
+	void FRivermaxOutStream::Stop()
+	{
+		bIsActive = false;
+	}
+
+	void FRivermaxOutStream::Exit()
+	{
+
+	}
+
+	void FRivermaxOutStream::PrepareNextFrame()
+	{
+		TRACE_CPUPROFILER_EVENT_SCOPE(PrepareNextFrame);
+
+		switch (Options.AlignmentMode)
+		{
+			case ERivermaxAlignmentMode::FrameCreation:
+			{
+				PrepareNextFrame_FrameCreation();
+				break;
+			}
+			case ERivermaxAlignmentMode::AlignmentPoint:
+			{
+				PrepareNextFrame_AlignmentPoint();
+				break;
+			}
+			default:
+			{
+				checkNoEntry();
+			}
+		}
+	}
+
+	void FRivermaxOutStream::PrepareNextFrame_FrameCreation()
+	{
+		// When aligning on frame creation, we will always wait for a frame to be available.
+		TRACE_CPUPROFILER_EVENT_SCOPE(Rmax::WaitForReadyFrame);
+		TSharedPtr<FRivermaxOutputFrame> NextFrameToSend = GetNextFrameToSend(true);
+
+		// In frame creation alignment, we always release the last frame sent
+		if (CurrentFrame.IsValid())
+		{
+			constexpr bool bReleaseFrame = true;
+			CompleteCurrentFrame(bReleaseFrame);
+		}
+
+		// Make the next frame to send the current one and update its state
+		if (NextFrameToSend)
+		{
+			CurrentFrame = MoveTemp(NextFrameToSend);
+			InitializeNextFrame(CurrentFrame);
+		}
+	}
+
+	void FRivermaxOutStream::PrepareNextFrame_AlignmentPoint()
+	{
+		TRACE_CPUPROFILER_EVENT_SCOPE(Rmax::GetNextFrame_AlignmentPoint);
+
+		// When aligning on alignment points:
+		// We prepare to send the next frame that is ready if there is one available
+		// if none are available and bDoContinuousOutput == true
+		//		Repeat the last frame
+		// if none are available and bDoContinuousOutput == false
+		//		Don't send a frame and go back waiting for the next alignment point
+		
+		TSharedPtr<FRivermaxOutputFrame> NextFrameToSend = GetNextFrameToSend();
+
+		// If we have a new frame, release the previous one
+		// If we don't have a frame and we're not doing continuous output, we release it and we won't send a new one
+		// If we don't have a frame but we are doing continuous output, we will reschedule the current one, so no release.
+		if (!Options.bDoContinuousOutput || NextFrameToSend)
+		{
+			if (CurrentFrame)
+			{
+				constexpr bool bReleaseFrame = true;
+				CompleteCurrentFrame(bReleaseFrame);
+			}
+
+			// Make the next frame to send the current one and update its state
+			if (NextFrameToSend)
+			{
+				CurrentFrame = MoveTemp(NextFrameToSend);
+				InitializeNextFrame(CurrentFrame);
+			}
+		}
+		else
+		{
+			// We finished sending a frame so complete it but don't release it as we will repeat it
+			constexpr bool bReleaseFrame = false;
+			CompleteCurrentFrame(bReleaseFrame);
+
+			// We will resend the last one so just reinitialize it to resend
+			InitializeNextFrame(CurrentFrame);
+			
+			// If intermediate buffer isn't used and frame has to be repeated, we use skip chunk method 
+			// which might cause timing errors caused by chunk managements issue 
+			if (!StreamMemory.bUseIntermediateBuffer)
+			{
+				// No frame to send, keep last one and restart its internal counters
+				UE_LOG(LogRivermax, Verbose, TEXT("No frame to send. Reusing last frame with Frame Counter: %u"), CurrentFrame->GetFrameCounter());
+
+				// Since we want to resend last frame, we need to fast forward chunk pointer to re-point to the one we just sent
+				SkipChunks(StreamMemory.ChunksPerFrameField * (Options.NumberOfBuffers - 1));
+			}
+		}
+	}
+
+	void FRivermaxOutStream::DestroyStream_Internal()
+	{
+		DestroyStream();
+	}
+
+	void FRivermaxOutStream::InitializeStreamTimingSettings()
+	{
+		using namespace UE::RivermaxCore::Private::Utils;
+
+		const double TROOverride = CVarRivermaxOutputTROOverride.GetValueOnAnyThread();
+		if (TROOverride != 0)
+		{
+			TransmitOffsetNanosec = TROOverride * 1E9;
+			return;
+		}
+	}
+
+	const FFrameRate& FRivermaxOutStream::GetFrameRate() const
+	{
+		return Options.StreamOptions[static_cast<uint8>(StreamType)]->FrameRate;
+	}
+
+	const FString& FRivermaxOutStream::GetStreamAddress() const
+	{
+		return Options.StreamOptions[static_cast<uint8>(StreamType)]->StreamAddress;
+	}
+
+	const FString& FRivermaxOutStream::GetInterfaceAddress() const
+	{
+		return Options.StreamOptions[static_cast<uint8>(StreamType)]->InterfaceAddress;
+	}
+
+	uint32 FRivermaxOutStream::GetPort() const
+	{
+		return Options.StreamOptions[static_cast<uint8>(StreamType)]->Port;
+	}
+
+	uint64 FRivermaxOutStream::GetStreamIndexSDP() const
+	{
+		return Options.StreamOptions[static_cast<uint8>(StreamType)]->StreamIndex;
+	}
+
+	void FRivermaxOutStream::LogStreamDescriptionOnCreation() const
+	{
+		TStringBuilder<512> StreamDescription;
+
+		StreamDescription.Appendf(TEXT("Output stream started sending on stream %s:%d using interface %s. ")
+			, *GetStreamAddress()
+			, GetPort()
+			, *GetInterfaceAddress());
+
+		// Matches ERivermaxStreamType
+		FString StreamTypeToStrMap[] =
+		{
+			"Video",
+			"Audio",
+			"Ancillary"
+		};
+
+		StreamDescription.Appendf(TEXT("StreamType = %s "), *StreamTypeToStrMap[static_cast<uint8>(StreamType)]);
+		UE_LOG(LogRivermax, Display, TEXT("%s"), *FString(StreamDescription));
+	}
+
+
+	void FRivermaxOutStream::ShowStats()
+	{
+		if (CachedCVars.bShowOutputStats)
+		{
+			const double CurrentTime = FPlatformTime::Seconds();
+			if (CurrentTime - LastStatsShownTimestamp > CachedCVars.ShowOutputStatsIntervalSeconds)
+			{
+				LastStatsShownTimestamp = CurrentTime;
+				UE_LOG(LogRivermax, Log, TEXT("Stats: FrameSent: %llu. CommitImmediate: %u. CommitRetries: %u. ChunkRetries: %u. ChunkSkippingRetries: %u. Timing issues: %llu"), Stats.FramesSentCounter, Stats.CommitImmediate, Stats.CommitRetries, Stats.TotalChunkRetries, Stats.ChunkSkippingRetries, Stats.TimingIssueCount);
+			}
+		}
+	}
+
+	bool FRivermaxOutStream::IsGPUDirectSupported() const
+	{
+		return bUseGPUDirect;
+	}
+
+	void FRivermaxOutStream::CalculateNextScheduleTime_AlignementPoints(uint64 CurrentClockTimeNanosec, uint64 CurrentFrameNumber)
+	{
+		// Frame number we will want to align with
+		uint64 NextFrameNumber = CurrentFrameNumber;
+
+		bool bFoundValidTimings = true;
+		
+		if (StreamData.bHasValidNextFrameNumber == false)
+		{
+			// Now that the stream starts when a frame was produced, we can reduce our wait
+			// We wait one frame here to start sending at the next frame boundary.
+			// Since it takes a frame to send it, we could detect if we are in the first 10% (arbitrary)
+			// of the interval and start sending right away but we might be overlapping with the next one
+			NextFrameNumber = CurrentFrameNumber + 1;
+		}
+		else
+		{	
+			// Case where we are back and frame number is the previous one. Depending on offsets, this could happen
+			if (CurrentFrameNumber == StreamData.NextAlignmentPointFrameNumber - 1)
+			{
+				NextFrameNumber = StreamData.NextAlignmentPointFrameNumber + 1;
+				UE_LOG(LogRivermax, Verbose, TEXT("Scheduling last frame was faster than expected. (CurrentFrame: '%llu' LastScheduled: '%llu') Scheduling for following expected one.")
+					, CurrentFrameNumber
+					, StreamData.NextAlignmentPointFrameNumber);
+			}
+			else
+			{
+				// We expect current frame number to be the one we scheduled for the last time or greater if something happened
+				if (CurrentFrameNumber >= StreamData.NextAlignmentPointFrameNumber)
+				{
+					// Verify if last frame had timing issues. If yes, we skip next interval
+					if (CurrentFrame && CurrentFrame->bCaughtTimingIssue)
+					{
+						NextFrameNumber = CurrentFrameNumber + 2;
+						UE_LOG(LogRivermax, Warning, TEXT("Timing issue detected during frame %llu. Skipping frame %llu to keep sync."), CurrentFrameNumber, CurrentFrameNumber + 1);
+					}
+					else
+					{
+						// If current frame is greater than last scheduled, we missed an alignment point.
+						const uint64 DeltaFrames = CurrentFrameNumber - StreamData.NextAlignmentPointFrameNumber;
+						if (DeltaFrames >= 1)
+						{
+							UE_LOG(LogRivermax, Warning, TEXT("Output missed %llu frames."), DeltaFrames);
+
+							// If we missed a sync point, this means that last scheduled frame might still be ongoing and 
+							// sending it might be crossing the frame boundary so we skip one entire frame to empty the queue.
+							NextFrameNumber = CurrentFrameNumber + 2;
+						}
+						else
+						{
+							NextFrameNumber = CurrentFrameNumber + 1;
+						}
+					}
+				}
+				else
+				{
+					// This is not expected (going back in time) but we should be able to continue. Scheduling immediately
+					ensureMsgf(false, TEXT("Unexpected behaviour during output stream's alignment point calculation. Current time has gone back in time compared to last scheduling."));
+					bFoundValidTimings = false;
+				}
+			}
+		}
+
+		// Get next alignment point based on the frame number we are aligning with
+		const uint64 NextAlignmentNano = UE::RivermaxCore::GetAlignmentPointFromFrameNumber(NextFrameNumber, GetFrameRate());
+
+		// Add Tro offset to next alignment point and configurable offset
+		StreamData.NextAlignmentPointNanosec = NextAlignmentNano;
+		StreamData.NextScheduleTimeNanosec = NextAlignmentNano + TransmitOffsetNanosec + CVarRivermaxScheduleOffset.GetValueOnAnyThread();
+		StreamData.LastAlignmentPointFrameNumber = StreamData.NextAlignmentPointFrameNumber;
+		StreamData.NextAlignmentPointFrameNumber = NextFrameNumber;
+
+		StreamData.bHasValidNextFrameNumber = bFoundValidTimings;
+	}
+
+	void FRivermaxOutStream::CalculateNextScheduleTime_FrameCreation(uint64 CurrentClockTimeNanosec, uint64 CurrentFrameNumber)
+	{
+		double NextWaitTime = 0.0;
+		if (StreamData.bHasValidNextFrameNumber == false)
+		{
+			StreamData.NextAlignmentPointNanosec = CurrentClockTimeNanosec;
+			StreamData.NextScheduleTimeNanosec = StreamData.NextAlignmentPointNanosec + CVarRivermaxScheduleOffset.GetValueOnAnyThread();
+			StreamData.NextAlignmentPointFrameNumber = CurrentFrameNumber;
+			StreamData.bHasValidNextFrameNumber = true;
+		}
+		else
+		{
+			// In this mode, we just take last time we started to send and add a frame interval
+			StreamData.NextAlignmentPointNanosec = StreamData.LastSendStartTimeNanoSec + StreamData.FrameFieldTimeIntervalNs;
+			StreamData.NextScheduleTimeNanosec = StreamData.NextAlignmentPointNanosec + CVarRivermaxScheduleOffset.GetValueOnAnyThread();
+			StreamData.NextAlignmentPointFrameNumber = UE::RivermaxCore::GetFrameNumber(StreamData.NextAlignmentPointNanosec, GetFrameRate());
+		}
+	}
+
+	bool FRivermaxOutStream::ReserveFrame(uint64 FrameCounter) const
+	{
+		return false;
+	}
+
+	void FRivermaxOutStream::OnFrameReadyToBeSent()
+	{	
+		FrameReadyToSendSignal->Trigger();
+	}
+
+	void FRivermaxOutStream::OnFrameReadyToBeUsed()
+	{
+		FrameAvailableSignal->Trigger();
+	}
+
+	void FRivermaxOutStream::OnPreFrameReadyToBeSent()
+	{
+		Listener->OnPreFrameEnqueue();
+	}
+
+	void FRivermaxOutStream::OnFrameManagerCriticalError()
+	{
+		Listener->OnStreamError();
+		Stop();
+	}
+
+	void FRivermaxOutStream::CacheCVarValues()
+	{
+		CachedCVars.bEnableCommitTimeProtection = CVarRivermaxOutputEnableTimingProtection.GetValueOnAnyThread() != 0;
+		CachedCVars.ForceCommitImmediateTimeNanosec = CVarRivermaxOutputForceImmediateSchedulingThreshold.GetValueOnAnyThread();
+		CachedCVars.SkipSchedulingTimeNanosec = CVarRivermaxOutputSkipSchedulingCutOffTime.GetValueOnAnyThread() * 1E3;
+		CachedCVars.bUseSingleMemblock = CVarRivermaxOutputUseSingleMemblock.GetValueOnAnyThread() == 1;
+		CachedCVars.bEnableTimeCriticalThread = CVarRivermaxOutputEnableTimeCriticalThread.GetValueOnAnyThread() != 0;
+		CachedCVars.bShowOutputStats = CVarRivermaxOutputShowStats.GetValueOnAnyThread() != 0;
+		CachedCVars.ShowOutputStatsIntervalSeconds = CVarRivermaxOutputShowStatsInterval.GetValueOnAnyThread();
+		CachedCVars.bPrefillRTPHeaders = CVarRivermaxOutputPrefillRTPHeaders.GetValueOnAnyThread();
+	}
+
+	bool FRivermaxOutStream::IsChunkOnTime() const
+	{
+		switch (Options.AlignmentMode)
+		{
+		case ERivermaxAlignmentMode::AlignmentPoint:
+		{
+			return IsChunkOnTime_AlignmentPoints();
+		}
+		case ERivermaxAlignmentMode::FrameCreation:
+		{
+			return IsChunkOnTime_FrameCreation();
+		}
+		default:
+		{
+			checkNoEntry();
+			return false;
+		}
+		}
+	}
+
+	bool FRivermaxOutStream::IsChunkOnTime_FrameCreation() const
+	{
+		return true;
+	}
+
+	bool FRivermaxOutStream::IsChunkOnTime_AlignmentPoints() const
+	{
+		if (CachedCVars.bEnableCommitTimeProtection)
+		{
+			// Calculate at what time this chunk is supposed to be sent
+			const uint64 DeltaTimePerChunk = (StreamData.FrameFieldTimeIntervalNs - TransmitOffsetNanosec) / StreamMemory.ChunksPerFrameField;
+			const uint64 NextChunkCommitTime = StreamData.NextScheduleTimeNanosec + (CurrentFrame->ChunkNumber * DeltaTimePerChunk);
+
+			// Verify if we are on time to send it. Use CVar to tighten / extend needed window. 
+			// This is to avoid messing up timing
+			const uint64 CurrentTime = RivermaxModule->GetRivermaxManager()->GetTime();
+			if (NextChunkCommitTime <= (CurrentTime + CachedCVars.SkipSchedulingTimeNanosec))
+			{
+				TRACE_CPUPROFILER_EVENT_SCOPE(RmaxOut::ChunkTooLate);
+				return false;
+			}
+
+			// Add other causes of timing issues. 
+			// Possible options : Chunk warnings, Last commit time too close to frame boundary, etc...
+		}
+		
+		return true;
+	}
+
+	void FRivermaxOutStream::OnMemoryChunksCopied(const TSharedPtr<FBaseDataCopySideCar>& Sidecar)
+	{
+	}
+
+	void FRivermaxOutStream::OnCVarRandomDelayChanged(IConsoleVariable* Var)
+	{
+		bTriggerRandomDelay = true;
+	}
+
+	void FRivermaxOutStream::CalculateFrameTimestamp()
+	{
+		using namespace UE::RivermaxCore::Private::Utils;
+
+		// For now, in order to be able to use a framelocked input, we pipe frame number in the timestamp for a UE-UE interaction
+		// Follow up work to investigate adding this in RTP header
+		uint64 InputTime = StreamData.NextAlignmentPointNanosec;
+		if (Options.bDoFrameCounterTimestamping)
+		{
+			InputTime = UE::RivermaxCore::GetAlignmentPointFromFrameNumber(CurrentFrame->GetFrameCounter(), GetFrameRate());
+		}
+
+		CurrentFrame->MediaTimestamp = GetTimestampFromTime(InputTime, MediaClockSampleRate);
+
+		//TODO: Use engine timecode. Aggregated in thread safe place.
+		CurrentFrame->Timecode = GetTimecodeFromTime(InputTime, MediaClockSampleRate, GetFrameRate());
+	}
+
+	void FRivermaxOutStream::SkipChunks(uint64 ChunkCount)
+	{
+		bool bHasAddedTrace = false;
+		rmx_status Status;
+		do
+		{
+			checkSlow(CachedAPI);
+			Status = CachedAPI->rmx_output_media_skip_chunks(&StreamData.ChunkHandle, ChunkCount);
+			if (Status != RMX_OK)
+			{
+				if (Status == RMX_NO_FREE_CHUNK)
+				{
+					// Wait until there are enough free chunk to be skipped
+					if (!bHasAddedTrace)
+					{
+						UE_LOG(LogRivermax, Warning, TEXT("No chunks ready to skip. Waiting"));
+						TRACE_CPUPROFILER_EVENT_SCOPE(NoFreeChunk);
+						bHasAddedTrace = true;
+					}
+				}
+				else
+				{
+					ensure(false);
+					UE_LOG(LogRivermax, Error, TEXT("Invalid error happened while trying to skip chunks. Status: %d."), Status);
+					Listener->OnStreamError();
+					Stop();
+				}
+			}
+		} while (Status != RMX_OK && bIsActive);
+	}
+
+	void FRivermaxOutStream::SendFrame()
+	{
+		TRACE_CPUPROFILER_EVENT_SCOPE(SendFrame);
+		StreamData.LastSendStartTimeNanoSec = RivermaxModule->GetRivermaxManager()->GetTime();
+
+		if (bTriggerRandomDelay)
+		{
+			TRACE_CPUPROFILER_EVENT_SCOPE(FRivermaxOutStream::RandomDelay);
+			bTriggerRandomDelay = false;
+			FPlatformProcess::SleepNoStats(FMath::RandRange(2e-3, 4e-3));
+		}
+
+		// Calculate frame's timestamp only once and reuse in RTP build
+		CalculateFrameTimestamp();
+
+		const uint32 MediaFrameNumber = Utils::TimestampToFrameNumber(CurrentFrame->MediaTimestamp, GetFrameRate());
+		TRACE_CPUPROFILER_EVENT_SCOPE_TEXT(*FRivermaxTracingUtils::RmaxOutSendingFrameTraceEvents[MediaFrameNumber % 10]);
+
+		TRACE_CPUPROFILER_EVENT_SCOPE_TEXT(*FRivermaxTracingUtils::RmaxOutMediaCapturePipeTraceEvents[CurrentFrame->GetFrameCounter() % 10]);
+		UE_LOG(LogRivermax, VeryVerbose, TEXT("RmaxRX Sending frame number %u with timestamp %u."), MediaFrameNumber, CurrentFrame->MediaTimestamp);
+
+		do
+		{
+			if (bIsActive)
+			{
+				GetNextChunk();
+			}
+
+			if (bIsActive && StreamMemory.bUseIntermediateBuffer && ((CurrentFrame->ChunkNumber % StreamMemory.ChunkSpacingBetweenMemcopies) == 0))
+			{
+				CopyFrameData(CurrentFrame, reinterpret_cast<uint8*>(CurrentFrame->FrameStartPtr));
+			}
+
+			if (bIsActive)
+			{
+				SetupRTPHeaders();
+			}
+
+			if (bIsActive)
+			{
+				if (!CurrentFrame->bCaughtTimingIssue)
+				{
+					// As long as our frame is good, verify if we commit chunks before it is expected to be sent
+					// We keep committing the frame even if we detect timing issue to avoid having to skip
+					// chunks in the internals of Rivermax and keep it going for entirety of frames. 
+					// We skip an interval instead but it is quite drastic.
+					const bool bIsChunkOnTime = IsChunkOnTime();
+					CurrentFrame->bCaughtTimingIssue = !bIsChunkOnTime;
+
+					if (UE::RivermaxCore::Private::GbTriggerRandomTimingIssue)
+					{
+						FRandomStream RandomStream(FPlatformTime::Cycles64());
+						const bool bTriggerDesync = (RandomStream.FRandRange(0.0, 1.0) > 0.7) ? true : false;
+						if (bTriggerDesync)
+						{
+							TRACE_CPUPROFILER_EVENT_SCOPE(RmaxOut::ForceTimingIssue);
+							CurrentFrame->bCaughtTimingIssue = true;
+						}
+
+						UE::RivermaxCore::Private::GbTriggerRandomTimingIssue = false;
+					}
+				}
+
+				CommitNextChunks();
+			}
+
+			//Update frame progress
+			if (bIsActive)
+			{
+				Stats.TotalPacketSent += StreamMemory.PacketsPerChunk;
+				++CurrentFrame->ChunkNumber;
+			}
+		} while (CurrentFrame->ChunkNumber < StreamMemory.ChunksPerFrameField && bIsActive);
+
+		Stats.FramesSentCounter++;
+	}
+
+	void FRivermaxOutStream::GetLastPresentedFrame(FPresentedFrameInfo& OutFrameInfo) const
+	{
+		FScopeLock Lock(&PresentedFrameCS);
+		OutFrameInfo = LastPresentedFrame;
+	}
+
+	void FRivermaxOutStream::CompleteCurrentFrame(bool bReleaseFrame)
+	{
+		if (ensure(CurrentFrame))
+		{
+			{
+				FScopeLock Lock(&PresentedFrameCS);
+				LastPresentedFrame.PresentedFrameBoundaryNumber = StreamData.LastAlignmentPointFrameNumber;
+				LastPresentedFrame.RenderedFrameNumber = CurrentFrame->GetFrameCounter();
+			}
+
+			// We don't release when there is no new frame, so we keep a hold on it to repeat it.
+			if (bReleaseFrame)
+			{
+				CurrentFrame.Reset();
+			}
+		}
+	}
+}
+
diff --git a/Engine/Plugins/VirtualProduction/Rivermax/RivermaxCore/Source/RivermaxCore/Private/Streams/RivermaxOutStream.h b/Engine/Plugins/VirtualProduction/Rivermax/RivermaxCore/Source/RivermaxCore/Private/Streams/RivermaxOutStream.h
new file mode 100644
--- /dev/null
+++ b/Engine/Plugins/VirtualProduction/Rivermax/RivermaxCore/Source/RivermaxCore/Private/Streams/RivermaxOutStream.h
@@ -0,0 +1,446 @@
+// Copyright Epic Games, Inc. All Rights Reserved.
+
+#pragma once
+
+#include "IRivermaxOutputStream.h"
+
+#include "Async/Future.h"
+#include "Containers/SpscQueue.h"
+#include "HAL/Runnable.h"
+#include "RivermaxWrapper.h"
+#include "RivermaxOutputFrame.h"
+#include "RivermaxTypes.h"
+#include "RTPHeader.h"
+
+
+class FEvent;
+class IRivermaxCoreModule;
+
+namespace UE::RivermaxCore::Private
+{
+	class FFrameManager;
+	class FBaseFrameAllocator;
+	struct FBaseDataCopySideCar;
+
+	using UE::RivermaxCore::FRivermaxOutputOptions;
+	using UE::RivermaxCore::FRivermaxOutputVideoFrameInfo;
+
+
+	/** Struct holding configuration information with regards to stream memory and packetization */
+	struct FRivermaxOutputStreamMemory
+	{
+		/** 
+		* Size of each data payload packet that will be used. In case of single SRD this payload 
+		* size will be aligned to equally divisible parts of image line width.
+		*/
+		uint16 PayloadSize = 0;
+
+		/** Number of pixel group per packet */
+		uint32 PixelGroupPerPacket = 0;
+
+		/** Number of pixels per packet */
+		uint32 PixelsPerPacket = 0;
+
+		/** Number of pixels per frame */
+		uint32 PixelsPerFrame = 0;
+
+		/** Stride of RTP header data. */
+		uint32 HeaderStrideSize = 20;
+
+		/** Number of lines packed inside a chunk. Can be controlled with cvar */
+		uint32 LinesInChunk = 4;
+
+		/** Number of packets per line.  */
+		uint32 PacketsInLine = 0;
+
+		/** Number of packets per chunk. Depends on LinesInChunk */
+		uint32 PacketsPerChunk = 0;
+
+		/** Number of frames per memory block. */
+		uint32 FramesFieldPerMemoryBlock = 0;
+
+		/** Number of packets per frame */
+		uint32 PacketsPerFrame = 0;
+
+		/** Number of packets per memory block */
+		uint32 PacketsPerMemoryBlock = 0;
+
+		/**Number of chunks per frame  */
+		uint32 ChunksPerFrameField = 0;
+
+		/** Number of chunks per memory block */
+		uint32 ChunksPerMemoryBlock = 0;
+
+		/** Number of memory block */
+		uint32 MemoryBlockCount = 0; 
+		
+		/** Whether intermediate buffer is used and captured frame has to be copied over again. */
+		bool bUseIntermediateBuffer = false;
+
+		/** Number of slices we split frame data into when copying it into intermediate buffer */
+		uint32 FrameMemorySliceCount = 1;
+
+		/** Chunk committed between each memcopy of frame data. Helps respect timing. */
+		uint32 ChunkSpacingBetweenMemcopies = 1;
+
+		/** Memory blocks allocated by Rivermax which is where packet data is stored. */
+		TArray<rmx_output_media_mem_block> MemoryBlocks;
+
+		/** 
+		* Data Sub block ID. When Number of sub blocks is more than one (which is the case if dynamic header split is used), 
+		* the first sub block is reserved for headers as per API docs. 
+		*/
+		uint8 HeaderBlockID = 0;
+
+		/** Data Sub block ID. For ANC we don't need headers to be split, so DataBlockID will be set 0. */
+		uint8 DataBlockID = 1;
+
+		/** Array with each packet size */
+		TArray<uint16_t> PayloadSizes; 
+
+		/** Array with each RTP header size */
+		TArray<uint16_t> HeaderSizes;
+
+		/** Contains RTP headers per memory block */
+		TArray<TArray<FVideoRTPHeader>> RTPHeaders;
+
+		/** Start addresses of each buffer in memblock */
+		TArray<void*> BufferAddresses;
+	};
+
+	struct FRivermaxOutputStreamStats
+	{
+		/** Chunk retries that were required since stream was started */
+		uint32 TotalChunkRetries = 0;
+		
+		/** Chunk retries that happened during last frame */
+		uint32 LastFrameChunkRetries = 0;
+		
+		/** Chunk skipping retries that happened since stream was started */
+		uint32 ChunkSkippingRetries = 0;
+		
+		/** Total packets that have been sent since stream was started */
+		uint32 TotalPacketSent = 0;
+		
+		/** Number of retries that were required when committing and queue was full since stream was started */
+		uint32 CommitRetries = 0;
+
+		/** Immediate commits that were done because we got there too close to scheduling time */
+		uint32 CommitImmediate = 0;
+
+		/** Number of frames that were sent since stream was started */
+		uint64 FramesSentCounter = 0;
+
+		/** Frames that had timing issues since stream was started */
+		uint64 TimingIssueCount = 0;
+	};
+
+	struct FRivermaxOutputStreamData
+	{
+		// Handle used to retrieve chunks associated with output stream
+		rmx_output_media_chunk_handle ChunkHandle;
+
+		/** Current sequence number being done */
+		uint32 SequenceNumber = 0;
+
+		/**  The SynchronizationSource (SSRC) field in the RTP header is a 32-bit identifier used to group RTP packets into a single stream. */
+		uint32 SynchronizationSource = FRawRTPHeader::VideoSynchronizationSource;
+
+		/** Time interval between frames in nanoseconds. */
+		double FrameFieldTimeIntervalNs = 0.0;
+
+		/** Next alignment point based on PTP standard */
+		uint64 NextAlignmentPointNanosec = 0;
+
+		/** Next schedule time using 2110 gapped model timing and controllable offset */
+		uint64 NextScheduleTimeNanosec = 0;
+
+		/** Whether next alignment frame number is deemed valid or not to detect missed frames. */
+		bool bHasValidNextFrameNumber = false;
+		
+		/** Next alignment point frame number treated to detect missed frames */
+		uint64 NextAlignmentPointFrameNumber = 0;
+
+		/** Last alignment point frame number we have processed*/
+		uint64 LastAlignmentPointFrameNumber = 0;
+
+		/** Timestamp at which we started commiting a frame */
+		uint64 LastSendStartTimeNanoSec = 0;
+		
+		/** Keeping track of how much time was slept last round. */
+		uint64 LastSleepTimeNanoSec = 0;
+
+	};
+
+	/** Struct holding various cached cvar values that can't be changed once stream has been created and to avoid calling anythread getters continuously */
+	struct FOutputStreamCachedCVars
+	{
+		/** Whether timing protection is active and next frame interval is skipped if it happens */
+		bool bEnableCommitTimeProtection = true;
+
+		/** Time padding from scheduling time required to avoid skipping it */
+		uint64 SkipSchedulingTimeNanosec = 0;
+
+		/** 
+		 * Time from scheduling required to not commit it immediately 
+		 * Rivermax sdk will throw an error if time is in the past when it
+		 * gets to actually comitting it. 
+		 */
+		uint64 ForceCommitImmediateTimeNanosec = 0;
+		
+		/** Tentative optimization recommended for SDK where a single big memblock is allocated. When false, a memblock per frame is configured. */
+		bool bUseSingleMemblock = true;
+
+		/** Whether to bump output thread priority to time critical */
+		bool bEnableTimeCriticalThread = true;
+
+		/** Whether to show output stats at regular interval in logs */
+		bool bShowOutputStats = false;
+
+		/** Interval in seconds at which to display output stats */
+		float ShowOutputStatsIntervalSeconds = 1.0f;
+
+		/** Whether to prefill RTP header memory with known data at initialization time instead of during sending */
+		bool bPrefillRTPHeaders = true;
+	};
+
+	class FRivermaxOutStream : public UE::RivermaxCore::IRivermaxOutputStream, public FRunnable
+	{
+	public:
+		FRivermaxOutStream(const TArray<char>& SDPDescription);
+		virtual ~FRivermaxOutStream();
+
+	public:
+
+		//~ Begin IRivermaxOutputStream interface
+		virtual bool Initialize(const FRivermaxOutputOptions& Options, IRivermaxOutputStreamListener& InListener) override;
+		virtual void Uninitialize() override;
+
+		virtual bool IsGPUDirectSupported() const override;
+		virtual bool ReserveFrame(uint64 FrameCounter) const override;
+		virtual void GetLastPresentedFrame(FPresentedFrameInfo& OutFrameInfo) const override;
+		//~ End IRivermaxOutputStream interface
+
+		virtual void Process_AnyThread();
+
+		//~ Begin FRunnable interface
+		virtual bool Init() override;
+		virtual uint32 Run() override;
+		virtual void Stop() override;
+		virtual void Exit() override;
+		//~ End FRunnable interface
+
+	protected:
+
+		/** Configures chunks, packetizing, memory blocks of the stream */
+		virtual bool InitializeStreamMemoryConfig() = 0;
+
+		/** Resets NextFrame to be ready to send it out */
+		void InitializeNextFrame(const TSharedPtr<FRivermaxOutputFrame>& NextFrame);
+
+		/** Indicates that there is some data that is ready to be sent*/
+		virtual bool IsFrameAvailableToSend() = 0;
+
+		/** Query rivermax library for the next chunk to work with */
+		virtual void GetNextChunk();
+
+		/** Fills RTP header for all packets to be sent for this chunk */
+		virtual void SetupRTPHeaders() = 0;
+
+		/** Copies part of frame memory in next memblock's chunk to be sent out */
+		virtual bool CopyFrameData(const TSharedPtr<FRivermaxOutputFrame>& SourceFrame, uint8* DestinationBase) = 0;
+
+		/** Commits chunk to rivermax so they are scheduled to be sent */
+		virtual void CommitNextChunks();
+
+		/** Fetches next frame to send and prepares it for sending */
+		virtual void PrepareNextFrame();
+
+		/** Get next frame to be sent, stream type agnostic. Should wait until the frame is available if bWait. */
+		virtual TSharedPtr<FRivermaxOutputFrame> GetNextFrameToSend(bool bWait = false) = 0;
+
+		/** Returns next frame to send for frame creation alignment */
+		virtual void PrepareNextFrame_FrameCreation();
+
+		/** Returns next frame to send for alignement point method. Can return nullptr */
+		virtual void PrepareNextFrame_AlignmentPoint();
+
+		/** Sets up frame management taking care of allocation, special cuda handling, etc... */
+		virtual bool SetupFrameManagement() = 0;
+
+		/** Clean up frames */
+		virtual void CleanupFrameManagement() = 0;
+
+		/** Destroys rivermax stream. Will wait until it's ready to be destroyed */
+		virtual void DestroyStream();
+
+		/** When a frame has been sent (after frame interval), we update last presented frame tracking and optionally release it in the presentation queue */
+		virtual void CompleteCurrentFrame(bool bReleaseFrame);
+
+		/** Waits for the next point in time to send out a new frame. Returns true if it exited earlier with the next frame ready to be processed */
+		virtual bool WaitForNextRound();
+
+		/** Uses time before next frame interval to copy data from next ready frame to intermediate buffer */
+		virtual void PreprocessNextFrame();
+
+		/** Initializes timing setup for this stream. TRO, frame interval etc... */
+		virtual void InitializeStreamTimingSettings();
+
+	protected:
+		/** Get frame rate according to the SDP file. */
+		virtual const FFrameRate& GetFrameRate() const;
+
+		/** Get stream address according to the SDP file. */
+		virtual const FString& GetStreamAddress() const;
+
+		/** Get Interface Address (Physical Port address). */
+		virtual const FString& GetInterfaceAddress() const;
+
+		/** Get Port Number according to the SDP file. */
+		virtual uint32 GetPort() const;
+
+		/** Get Stream Index as it is ordered in the SDP file. */
+		virtual uint64 GetStreamIndexSDP() const;
+
+		/** When stream creation is successful this should log all relevant information. */
+		virtual void LogStreamDescriptionOnCreation() const = 0;
+
+	protected:
+
+		/** Destroys rivermax stream. Will wait until it's ready to be destroyed */
+		void DestroyStream_Internal();
+
+		/** Calculate next frame scheduling time for alignment points mode */
+		void CalculateNextScheduleTime_AlignementPoints(uint64 CurrentClockTimeNanosec, uint64 CurrentFrameNumber);
+		
+		/** Calculate next frame scheduling time for frame creation mode */
+		void CalculateNextScheduleTime_FrameCreation(uint64 CurrentClockTimeNanosec, uint64 CurrentFrameNumber);
+
+		/** Validates timing on every commit to see if we are respecting alignment */
+		bool IsChunkOnTime() const;
+		
+		/** Validates timing for frame creation alignment which always returns true. */
+		bool IsChunkOnTime_FrameCreation() const;
+		
+		/** Validates timing to make sure chunk to be committed are on time. 
+		 *  Once a chunk is late, timings are at risk and next frame will be skipped
+		 */
+		bool IsChunkOnTime_AlignmentPoints() const;
+
+		/** If enabled, print stats related to this stream */
+		void ShowStats();
+
+		/** Used to notify the listener that a frame is ready to be enqueued for transmission */
+		void OnPreFrameReadyToBeSent();
+		
+		/** Used to detect when a frame is now ready to be sent */
+		void OnFrameReadyToBeSent();
+
+		/** Used to know when a frame is ready to be used and receive new data */
+		void OnFrameReadyToBeUsed();
+
+		/** Used to detect when the frame manager has caught a critical error */
+		void OnFrameManagerCriticalError();
+
+		/** Used to cache cvars at initialization */
+		void CacheCVarValues();
+
+		/** Called back when copy request was completed by allocator */
+		void OnMemoryChunksCopied(const TSharedPtr<FBaseDataCopySideCar>& Sidecar);
+
+		/** Called when delay request cvar has been changed */
+		void OnCVarRandomDelayChanged(IConsoleVariable* Var);
+
+		/** Update frame's timestamp to be used when setting every RTP headers */
+		void CalculateFrameTimestamp();
+
+		/** Tells Rivermax to skip a certain number of chunks in memory. Can be zero to just reset internals */
+		void SkipChunks(uint64 ChunkCount);
+
+		/** Go through all chunks of current frame and commit them to Rivermax to send them at the next desired time */
+		void SendFrame();
+
+	protected:
+
+		/** Rivermax memory configuration. i.e. memblock, chunks */
+		FRivermaxOutputStreamMemory StreamMemory;
+
+		/** Options related to this stream. i.e resolution, frame rate, etc... */
+		FRivermaxOutputOptions Options;
+
+		/** Various stats collected by this stream */
+		FRivermaxOutputStreamStats Stats;
+
+		/** State of various piece for this stream. Alignment points, schedule number, etc... */
+		FRivermaxOutputStreamData StreamData;
+
+		/** Stream id returned by rmax library */
+		rmx_stream_id StreamId = 0;
+
+		/** Current frame being sent */
+		TSharedPtr<FRivermaxOutputFrame> CurrentFrame;
+
+		/** 
+		* Frames reserved when Media Capture texture is about to be converted into Buffer to be sent.
+		* Key is frame counter
+		*/
+		mutable TQueue<TSharedPtr<FRivermaxOutputFrame>> ReservedFrames;
+
+		/** Thread scheduling frame output */
+		TUniquePtr<FRunnableThread> RivermaxThread;
+
+		/** Whether stream is active or not */
+		std::atomic<bool> bIsActive;
+
+		/** Event used to let scheduler that a frame is ready to be sent */
+		FEventRef FrameReadyToSendSignal = FEventRef(EEventMode::AutoReset);
+
+		/** Event used to unblock frame reservation as soon as one is free */
+		FEventRef FrameAvailableSignal = FEventRef(EEventMode::AutoReset);
+
+		/** Listener for this stream events */
+		IRivermaxOutputStreamListener* Listener = nullptr;
+
+		/** Type of stream created. Only 2110-20 (Video is supported now) */
+		ERivermaxStreamType StreamType = ERivermaxStreamType::VIDEO_2110_20_STREAM;
+
+		/** TRoffset time calculated based on ST2110 - 21 Gapped(for now) method. This is added to next alignment point */
+		uint64 TransmitOffsetNanosec = 0;
+
+		/** Timestamp at which we logged stats */
+		double LastStatsShownTimestamp = 0.0;
+		
+		/** Whether stream is using gpudirect to host memory consumed by Rivermax */
+		bool bUseGPUDirect = false;
+
+		/** Our own module pointer kept for ease of use */
+		IRivermaxCoreModule* RivermaxModule = nullptr;
+
+		/** Guid given by boundary monitoring handler to unregister ourselves */
+		// TODO: There must be a better place for this.
+		FGuid MonitoringGuid;
+
+		/** Future returned by the async initialization job we launch. Used to detect if it has completed during shutdown. */
+		TFuture<void> InitializationFuture;
+
+		/** Cached cvar values */
+		FOutputStreamCachedCVars CachedCVars;
+
+		/* Pointer to the rivermax API to avoid virtual calls in a hot loop. */ 
+		const RIVERMAX_API_FUNCTION_LIST* CachedAPI = nullptr;
+
+		/** Whether to trigger a delay in the output thread loop next time it ticks */
+		bool bTriggerRandomDelay = false;
+
+		/** Critical section to access data of last presented frame */
+		mutable FCriticalSection PresentedFrameCS;
+
+		/** Info of last presented frame */
+		FPresentedFrameInfo LastPresentedFrame;
+
+		/** An actual string of the SDP file. */
+		TArray<char> SDPDescription;
+	};
+}
+
+
diff --git a/Engine/Plugins/VirtualProduction/Rivermax/RivermaxCore/Source/RivermaxCore/Private/Streams/RivermaxOutVideoStream.cpp b/Engine/Plugins/VirtualProduction/Rivermax/RivermaxCore/Source/RivermaxCore/Private/Streams/RivermaxOutVideoStream.cpp
new file mode 100644
--- /dev/null
+++ b/Engine/Plugins/VirtualProduction/Rivermax/RivermaxCore/Source/RivermaxCore/Private/Streams/RivermaxOutVideoStream.cpp
@@ -0,0 +1,1013 @@
+// Copyright Epic Games, Inc. All Rights Reserved.
+
+#include "RivermaxOutVideoStream.h"
+
+#include "Async/Async.h"
+#include "CudaModule.h"
+#include "IRivermaxCoreModule.h"
+#include "IRivermaxBoundaryMonitor.h"
+#include "IRivermaxManager.h"
+#include "Misc/ByteSwap.h"
+#include "RivermaxFrameAllocator.h"
+#include "RivermaxFrameManager.h"
+#include "RivermaxLog.h"
+#include "RivermaxPTPUtils.h"
+#include "RivermaxTracingUtils.h"
+#include "RivermaxTypes.h"
+#include "RivermaxUtils.h"
+
+#define GET_FRAME_INDEX(OUTPUT_FRAME) OUTPUT_FRAME->GetFrameCounter() % Options.NumberOfBuffers
+namespace UE::RivermaxCore::Private
+{
+	static TAutoConsoleVariable<int32> CVarRivermaxOutputEnableMultiSRD(
+		TEXT("Rivermax.Output.EnableMultiSRD"), 1,
+		TEXT("When enabled and if the row cannot be split evenly, non-uniform payloads will be used. The last packet for the frame will not be fully filled with data.\n"
+			"If disabled, the payloads will be split evenly or the 2110 stream will be disabled."),
+		ECVF_Default);
+
+	static TAutoConsoleVariable<int32> CVarRivermaxOutputLinesPerChunk(
+		TEXT("Rivermax.Output.LinesPerChunk"), 4,
+		TEXT("Defines the number of lines to pack in a chunk. Higher number will increase latency"),
+		ECVF_Default);
+
+	static TAutoConsoleVariable<int32> CVarRivermaxOutputMaximizePacketSize(
+		TEXT("Rivermax.Output.MaximizePacketSize"), 1,
+		TEXT("Enables bigger packet sizes to maximize utilisation of potential UDP packet. If not enabled, packet size will be aligned with HD/4k sizes"),
+		ECVF_Default);
+
+	static TAutoConsoleVariable<int32> CVarRivermaxOutputMaxFrameMemorySliceCount(
+		TEXT("Rivermax.Output.FrameSliceCount"), 30,
+		TEXT("Max number of memcopies done per frame when using intermediate buffer. As frame gets bigger, we can't do a single memcopy or timings will be broken. Can be smaller in order to fit inside chunk count."),
+		ECVF_Default);
+
+	static TAutoConsoleVariable<bool> CVarRivermaxOutputEnableIntermediateBuffer(
+		TEXT("Rivermax.Output.Alignment.EnableIntermediateBuffer"), true,
+		TEXT("Uses an intermediate buffer used by Rivermax when sending data out.\n")
+		TEXT("During scheduling, captured frame data will be copied over intermediate buffer.\n")
+		TEXT("Only applies to alignment points scheduling mode."),
+		ECVF_Default);
+
+	static TAutoConsoleVariable<int32> CVarRivermaxOutputMemcopyChunkSpacing(
+		TEXT("Rivermax.Output.Scheduling.MemcopyChunkSpacing"), 10,
+		TEXT("Number of chunks between each memcopy to help with timing for different frame format."),
+		ECVF_Default);
+
+	bool FindPayloadSize(uint32 InBytesPerLine, const uint32 PixelGroupSize, uint16& OutPayloadSize)
+	{
+		using namespace UE::RivermaxCore::Private::Utils;
+
+		// For now, we only find a payload size that can be equal across one line
+		// Support for last payload of a line being smaller is there but is causing issue
+		// We fail to output if we receive a resolution for which we can't find an equal payload size
+		int32 TestPoint = InBytesPerLine / MaxPayloadSize;
+		if (TestPoint == 0)
+		{
+			if (InBytesPerLine > MinPayloadSize)
+			{
+				if (InBytesPerLine % PixelGroupSize == 0)
+				{
+					OutPayloadSize = InBytesPerLine;
+					return true;
+				}
+			}
+			return false;
+		}
+
+		while (true)
+		{
+			const int32 TestSize = InBytesPerLine / TestPoint;
+			if (TestSize < MinPayloadSize)
+			{
+				break;
+			}
+
+			if (TestSize <= MaxPayloadSize)
+			{
+				if ((TestSize % PixelGroupSize) == 0 && (InBytesPerLine % TestPoint) == 0)
+				{
+					OutPayloadSize = TestSize;
+					return true;
+				}
+			}
+
+			++TestPoint;
+		}
+
+		return false;
+	}
+
+	void UpdateVideoRtpHeader
+		( FVideoRTPHeader& OutVideoHeader
+		, uint32& OutPacketCounter
+		, uint16& OutSRDOffset
+		, uint32& OutLineNumber
+		, uint32 InSequenceNumber
+		, uint32 InTimeStamp
+		, const FIntPoint& AlignedResolution
+		, const FVideoFormatInfo& InFormatInfo
+		, const FRivermaxOutputStreamMemory& InStreamMemory
+	)
+	{
+		OutVideoHeader = {};
+		OutVideoHeader.RTPHeader = {};
+		OutVideoHeader.RTPHeader.Version = 2;
+		OutVideoHeader.RTPHeader.PaddingBit = 0;
+		OutVideoHeader.RTPHeader.ExtensionBit = 0;
+		OutVideoHeader.RTPHeader.PayloadType = 96; //Payload type should probably be inferred from SDP
+
+		OutVideoHeader.RTPHeader.SequenceNumber = ByteSwap((uint16)(InSequenceNumber & 0xFFFF));
+		OutVideoHeader.RTPHeader.Timestamp = ByteSwap(InTimeStamp);
+		OutVideoHeader.RTPHeader.ExtendedSequenceNumber = ByteSwap((uint16)((InSequenceNumber >> 16) & 0xFFFF));
+
+		//2110 specific header
+		OutVideoHeader.RTPHeader.SynchronizationSource = ByteSwap((uint32)0x0eb51dbd);  // Should Unreal has its own synch source ID
+
+		{
+			if (OutPacketCounter + 1 == InStreamMemory.PacketsPerFrame)
+			{
+				OutVideoHeader.RTPHeader.MarkerBit = 1; // last packet in frame (Marker)
+			}
+
+			// Verify if payload size exceeds line 
+			const uint32 CurrentPayloadSize = InStreamMemory.PayloadSizes[OutPacketCounter];
+
+			const uint32 LineSizeOffset = ((OutSRDOffset / InFormatInfo.PixelGroupCoverage) * InFormatInfo.PixelGroupSize);
+			const uint32 LineSize = ((AlignedResolution.X / InFormatInfo.PixelGroupCoverage) * InFormatInfo.PixelGroupSize);
+
+			const uint16 SRD1Length = FMath::Min(LineSize - LineSizeOffset, CurrentPayloadSize);
+			const uint16 SRD1PixelCount = SRD1Length / InFormatInfo.PixelGroupSize * InFormatInfo.PixelGroupCoverage;
+			uint16 SRD2Length = SRD1Length < CurrentPayloadSize ? CurrentPayloadSize - SRD1Length : 0;
+			if (SRD2Length && OutLineNumber == ((uint32)AlignedResolution.Y - 1))
+			{
+				SRD2Length = 0;
+			}
+
+			OutVideoHeader.SRD1.Length = ByteSwap(SRD1Length);
+			OutVideoHeader.SRD1.SetRowNumber(OutLineNumber); //todo divide by 2 if interlaced
+			OutVideoHeader.SRD1.SetOffset(OutSRDOffset);
+			OutVideoHeader.SRD1.ContinuationBit = SRD2Length > 0;
+			OutVideoHeader.SRD1.FieldIdentification = 0; //todo when fields are sent for interlace
+
+			OutSRDOffset += SRD1PixelCount;
+			if (OutSRDOffset >= AlignedResolution.X)
+			{
+				OutSRDOffset = 0;
+				++OutLineNumber;
+			}
+
+			if (SRD2Length > 0)
+			{
+				OutVideoHeader.SRD2.Length = ByteSwap(SRD2Length);
+				OutVideoHeader.SRD2.SetRowNumber(OutLineNumber);
+				OutVideoHeader.SRD2.SetOffset(OutSRDOffset);
+				OutVideoHeader.SRD2.ContinuationBit = 0;
+				OutVideoHeader.SRD2.FieldIdentification = 0;
+
+				const uint16 SRD2PixelCount = SRD2Length / InFormatInfo.PixelGroupSize * InFormatInfo.PixelGroupCoverage;
+				OutSRDOffset += SRD2PixelCount;
+				if (OutSRDOffset >= AlignedResolution.X)
+				{
+					OutSRDOffset = 0;
+					++OutLineNumber;
+				}
+			}
+		}
+	}
+
+	uint32 FindLinesPerChunk(const FRivermaxOutputOptions& InOptions)
+	{
+		// More lines per chunks mean we will do more work prior to start sending a chunk. So, added 'latency' in terms of packet / parts of frame.
+		// Less lines per chunk mean that sender thread might starve.
+		// SDK sample uses 4 lines for UHD and 8 for HD. 
+		return CVarRivermaxOutputLinesPerChunk.GetValueOnAnyThread();
+	}
+
+	uint16 GetPayloadSize(ESamplingType SamplingType)
+	{
+		const FVideoFormatInfo FormatInfo = FStandardVideoFormat::GetVideoFormatInfo(SamplingType);
+		uint16 PayloadSize;
+		switch (SamplingType)
+		{
+
+		case ESamplingType::YUV444_10bit:
+		{
+			// Passthrough
+		}
+		case ESamplingType::RGB_10bit:
+		{
+			PayloadSize = 1200;
+			break;
+		}
+		case ESamplingType::YUV444_8bit:
+		{
+			// Passthrough
+		}
+		case ESamplingType::RGB_8bit:
+		{
+			PayloadSize = 1152;
+			break;
+		}
+		case ESamplingType::YUV444_12bit:
+		{
+			// Passthrough
+		}
+		case ESamplingType::RGB_12bit:
+		{
+			PayloadSize = 1152;
+			break;
+		}
+		case ESamplingType::YUV444_16bit:
+		{
+			// Passthrough
+		}
+		case ESamplingType::YUV444_16bitFloat:
+		{
+			// Passthrough
+		}
+		case ESamplingType::RGB_16bit:
+		{
+			// Passthrough
+		}
+		case ESamplingType::RGB_16bitFloat:
+		{
+			PayloadSize = 1152;
+			break;
+		}
+		case ESamplingType::YUV422_8bit:
+		{
+			PayloadSize = 1280;
+			break;
+		}
+		case ESamplingType::YUV422_10bit:
+		{
+			PayloadSize = 1200;
+			break;
+		}
+		case ESamplingType::YUV422_12bit:
+		{
+			PayloadSize = 1152;
+			break;
+		}
+		case ESamplingType::YUV422_16bit:
+		{
+			// Passthrough
+		}
+		case ESamplingType::YUV422_16bitFloat:
+		{
+			PayloadSize = 1280;
+			break;
+		}
+		default:
+		{
+			checkNoEntry();
+			PayloadSize = 1200;
+			break;
+		}
+		}
+
+		ensure(PayloadSize % FormatInfo.PixelGroupSize == 0);
+		return PayloadSize;
+	}
+
+	/**
+	 * Returns a payload closer to the max value we can have for standard UDP size
+	 * RTPHeader can be bigger depending on configuration so we'll cap payload at 1400.
+	 */
+	uint16 GetMaximizedPayloadSize(ESamplingType SamplingType)
+	{
+		const FVideoFormatInfo FormatInfo = FStandardVideoFormat::GetVideoFormatInfo(SamplingType);
+		uint16 PayloadSize;
+		switch (SamplingType)
+		{
+
+		case ESamplingType::YUV444_10bit:
+		{
+			// Passthrough
+		}
+		case ESamplingType::RGB_10bit:
+		{
+			PayloadSize = 1395;
+			break;
+		}
+		case ESamplingType::YUV444_8bit:
+		{
+			// Passthrough
+		}
+		case ESamplingType::RGB_8bit:
+		{
+			PayloadSize = 1398;
+			break;
+		}
+		case ESamplingType::YUV444_12bit:
+		{
+			// Passthrough
+		}
+		case ESamplingType::RGB_12bit:
+		{
+			PayloadSize = 1395;
+			break;
+		}
+		case ESamplingType::YUV444_16bit:
+		{
+			// Passthrough
+		}
+		case ESamplingType::YUV444_16bitFloat:
+		{
+			// Passthrough
+		}
+		case ESamplingType::RGB_16bit:
+		{
+			// Passthrough
+		}
+		case ESamplingType::RGB_16bitFloat:
+		{
+			PayloadSize = 1398;
+			break;
+		}
+		case ESamplingType::YUV422_8bit:
+		{
+			PayloadSize = 1400;
+			break;
+		}
+		case ESamplingType::YUV422_10bit:
+		{
+			PayloadSize = 1400;
+			break;
+		}
+		case ESamplingType::YUV422_12bit:
+		{
+			PayloadSize = 1398;
+			break;
+		}
+		case ESamplingType::YUV422_16bit:
+		{
+			// Passthrough
+		}
+		case ESamplingType::YUV422_16bitFloat:
+		{
+			PayloadSize = 1400;
+			break;
+		}
+		default:
+		{
+			checkNoEntry();
+			PayloadSize = 1200;
+			break;
+		}
+		}
+
+		ensure(PayloadSize % FormatInfo.PixelGroupSize == 0);
+		return PayloadSize;
+	}
+
+
+	FRivermaxOutVideoStream::FRivermaxOutVideoStream(const TArray<char>& SDPDescription)
+		: FRivermaxOutStream(SDPDescription)
+	{
+	}
+
+	struct FRTPHeaderPrefiller
+	{
+		FRTPHeaderPrefiller(FRivermaxOutVideoStream& InRmaxOutputStream)
+			: Stream(InRmaxOutputStream)
+		{
+			RunningSRDOffsetPerFrame.SetNumZeroed(Stream.Options.NumberOfBuffers);
+			RunningLineNumberPerFrame.SetNumZeroed(Stream.Options.NumberOfBuffers);
+		}
+
+		void Update(uint32 PacketIndex)
+		{
+			using namespace UE::RivermaxCore::Private::Utils;
+
+			const int32 PacketCount = Stream.StreamMemory.ChunksPerFrameField * Stream.StreamMemory.PacketsPerChunk;
+			uint32 BufferIndex = 0;
+			for (int32 MemblockIndex = 0; MemblockIndex < Stream.StreamMemory.RTPHeaders.Num(); ++MemblockIndex)
+			{
+				for (uint32 FrameInBlockIndex = 0; FrameInBlockIndex < Stream.StreamMemory.FramesFieldPerMemoryBlock; ++FrameInBlockIndex)
+				{
+					uint16& SRDOffset = RunningSRDOffsetPerFrame[BufferIndex];
+					uint32& LineNumber = RunningLineNumberPerFrame[BufferIndex];
+					const uint32 HeaderIndex = PacketIndex + (FrameInBlockIndex * PacketCount);
+					if (ensure(Stream.StreamMemory.RTPHeaders[MemblockIndex].IsValidIndex(HeaderIndex)))
+					{
+						FVideoRTPHeader& CurrentHeader = Stream.StreamMemory.RTPHeaders[MemblockIndex][HeaderIndex];
+
+						UpdateVideoRtpHeader(
+							CurrentHeader,
+							PacketIndex,
+							SRDOffset,
+							LineNumber,
+							0, //SequenceNumber
+							0, //Timestamp
+							Stream.Options.GetStreamOptions<FRivermaxVideoOutputOptions>(Stream.StreamType)->AlignedResolution,
+							Stream.FormatInfo,
+							Stream.StreamMemory
+						);
+					}
+
+					++BufferIndex;
+				}
+			}
+		}
+
+	private:
+		TArray<uint16> RunningSRDOffsetPerFrame;
+		TArray<uint32> RunningLineNumberPerFrame;
+		FRivermaxOutVideoStream& Stream;
+	};
+
+	bool FRivermaxOutVideoStream::PushFrame(TSharedPtr<IRivermaxOutputInfo> FrameInfo)
+	{
+		TRACE_CPUPROFILER_EVENT_SCOPE(FRivermaxOutVideoStream::PushVideoFrame);
+		TSharedPtr<FRivermaxOutputVideoFrameInfo> NewFrameInfo = StaticCastSharedPtr<FRivermaxOutputVideoFrameInfo>(FrameInfo);
+		NewFrameInfo->Stride = GetRowSizeInBytes();
+		NewFrameInfo->Height = Options.GetStreamOptions<FRivermaxVideoOutputOptions>(StreamType)->AlignedResolution.Y;
+
+		// Clear reserved frame if there is one. If not, try to get free frame.
+		// ReservedFrame should always be valid when block on reservation mode is used.
+		TSharedPtr<FRivermaxOutputFrame> ReservedFrame;
+		if (ReservedFrames.Dequeue(ReservedFrame))
+		{
+			// Sanity check.
+			//check(FrameInfo->FrameIdentifier == ReservedFrame->GetFrameCounter());
+		}
+		else
+		{
+			check(Options.FrameLockingMode != EFrameLockingMode::BlockOnReservation);
+			ReservedFrame = FrameManager->GetFreeFrame();
+		}
+
+		// If this is invalid it means that frame locking mode is BlockOnReservation and the render ran faster than media output fps.
+		if (ReservedFrame.IsValid())
+		{
+			return FrameManager->SetFrameData(NewFrameInfo, ReservedFrame);
+		}
+		return false;
+	}
+
+	bool FRivermaxOutVideoStream::InitializeStreamMemoryConfig()
+	{
+		using namespace UE::RivermaxCore::Private::Utils;
+		TSharedPtr<FRivermaxVideoOutputOptions> StreamOptions = Options.GetStreamOptions<FRivermaxVideoOutputOptions>(StreamType);
+		FormatInfo = FStandardVideoFormat::GetVideoFormatInfo(StreamOptions->PixelFormat);
+
+		// Verify resolution for sampling type
+		if (StreamOptions->AlignedResolution.X % FormatInfo.PixelGroupCoverage != 0)
+		{
+			UE_LOG(LogRivermax, Warning, TEXT("Can't create Rivermax Output Stream. Aligned horizontal resolution of %d doesn't align with pixel group coverage of %d."), StreamOptions->AlignedResolution.X, FormatInfo.PixelGroupCoverage);
+			return false;
+		}
+
+		const SIZE_T FrameSize = GetRowSizeInBytes() * StreamOptions->AlignedResolution.Y;
+		if (FrameSize == 0)
+		{
+			UE_LOG(LogRivermax, Warning, TEXT("Can't create Rivermax Output Stream. Frame size of 0 is invalid. Verify resolution."));
+			return false;
+		}
+
+		//2110 stream type
+		//We need to use the fullframe allocated size to compute the payload size.
+
+		const SIZE_T BytesPerRow = GetRowSizeInBytes();
+
+		// By default we want to divide the bytes evenly across packets. Some resolutions will require packets to be sized unevenly.
+		const bool bFoundPayload = FindPayloadSize(BytesPerRow, FormatInfo.PixelGroupSize, StreamMemory.PayloadSize);
+		if (bFoundPayload == false)
+		{
+			// Find out payload we want to use. Either we go the 'potential' multi SRD route or we restrict the stream based on supported resolutions.
+			if (CVarRivermaxOutputEnableMultiSRD.GetValueOnAnyThread() >= 1)
+			{
+				UE_LOG(LogRivermax, Log, TEXT("Due to resolution %dx%d, row data will be sent over multiple packets with varied sizes."), StreamOptions->AlignedResolution.X, StreamOptions->AlignedResolution.Y);
+				if (CVarRivermaxOutputMaximizePacketSize.GetValueOnAnyThread() >= 1)
+				{
+					StreamMemory.PayloadSize = GetMaximizedPayloadSize(FormatInfo.Sampling);
+				}
+				else
+				{
+					StreamMemory.PayloadSize = GetPayloadSize(FormatInfo.Sampling);
+				}
+			}
+			else
+			{
+				UE_LOG(LogRivermax, Warning, TEXT("Could not find payload size for desired resolution %dx%d for desired pixel format."
+					"If the intention is to use non standard resolutions, users might want to enable multi-srd support via Rivermax.Output.EnableMultiSRD."), StreamOptions->AlignedResolution.X, StreamOptions->AlignedResolution.Y);
+				return false;
+			}
+		}
+
+		// With payload size in hand, figure out how many packets we will need, how many chunks (group of packets) and configure descriptor arrays
+		const uint32 PixelCount = StreamOptions->AlignedResolution.X * StreamOptions->AlignedResolution.Y;
+		const uint64 FrameSizeInBytes = PixelCount / FormatInfo.PixelGroupCoverage * FormatInfo.PixelGroupSize;
+
+		StreamMemory.PixelGroupPerPacket = StreamMemory.PayloadSize / FormatInfo.PixelGroupSize;
+		StreamMemory.PixelsPerPacket = StreamMemory.PixelGroupPerPacket * FormatInfo.PixelGroupCoverage;
+
+		// We might need a smaller packet to complete the end of frame so ceil to the next value
+		StreamMemory.PacketsPerFrame = FMath::CeilToInt32((float)PixelCount / StreamMemory.PixelsPerPacket);
+
+		// Depending on resolution and payload size, last packet of a line might not be fully utilized but we need the remaining bytes so ceil to next value
+		StreamMemory.PacketsInLine = FMath::CeilToInt32((float)StreamMemory.PacketsPerFrame / StreamOptions->AlignedResolution.Y);
+
+		StreamMemory.LinesInChunk = FindLinesPerChunk(Options);
+		StreamMemory.PacketsPerChunk = StreamMemory.LinesInChunk * StreamMemory.PacketsInLine;
+		StreamMemory.FramesFieldPerMemoryBlock = CachedCVars.bUseSingleMemblock ? Options.NumberOfBuffers : 1;
+
+		// Chunk count won't necessarily align with the number of packets required. We need an integer amount of chunks to initialize our stream
+		// and calculate how many packets that represents. Rivermax will expect the payload/header array to be that size. It just means that
+		// we will mark the extra packets as 0 size.
+		StreamMemory.ChunksPerFrameField = FMath::CeilToInt32((float)StreamMemory.PacketsPerFrame / StreamMemory.PacketsPerChunk);
+		const uint64 RealPacketsPerFrame = StreamMemory.ChunksPerFrameField * StreamMemory.PacketsPerChunk;
+		StreamMemory.PacketsPerMemoryBlock = RealPacketsPerFrame * StreamMemory.FramesFieldPerMemoryBlock;
+		StreamMemory.ChunksPerMemoryBlock = StreamMemory.FramesFieldPerMemoryBlock * StreamMemory.ChunksPerFrameField;
+		StreamMemory.MemoryBlockCount = Options.NumberOfBuffers / StreamMemory.FramesFieldPerMemoryBlock;
+
+		// Setup arrays with the right sizes so we can give pointers to rivermax. This makes the stream header sizes static.
+		StreamMemory.RTPHeaders.SetNumZeroed(StreamMemory.MemoryBlockCount);
+		StreamMemory.PayloadSizes.SetNumZeroed(StreamMemory.PacketsPerMemoryBlock);
+		StreamMemory.HeaderSizes.SetNumZeroed(StreamMemory.PacketsPerMemoryBlock);
+
+		// Packed size should be 26 octets.
+		StreamMemory.HeaderStrideSize = sizeof(FVideoRTPHeader);
+
+		if (!SetupFrameManagement())
+		{
+			return false;
+		}
+
+		FRTPHeaderPrefiller RTPFiller(*this);
+
+		// Used to keep track of acquired frames to prevent the same frame fomr being acquired from the pool more than once. 
+		// When this scope is exited all frames are returned back to the Frame Manager pool.
+		TArray<TSharedPtr<UE::RivermaxCore::Private::FRivermaxOutputFrame>> ProcessedFrames;
+
+		StreamMemory.MemoryBlocks.SetNumZeroed(StreamMemory.MemoryBlockCount);
+		CachedAPI->rmx_output_media_init_mem_blocks(StreamMemory.MemoryBlocks.GetData(), StreamMemory.MemoryBlockCount);
+		for (uint32 BlockIndex = 0; BlockIndex < StreamMemory.MemoryBlockCount; ++BlockIndex)
+		{
+			rmx_output_media_mem_block& Block = StreamMemory.MemoryBlocks[BlockIndex];
+			CachedAPI->rmx_output_media_set_chunk_count(&Block, StreamMemory.ChunksPerMemoryBlock);
+
+			// We have two sub block, header and data
+			constexpr uint8 SubBlockCount = 2;
+			CachedAPI->rmx_output_media_set_sub_block_count(&Block, SubBlockCount);
+
+			// Describe Header block
+			CachedAPI->rmx_output_media_set_packet_layout(&Block, StreamMemory.HeaderBlockID, StreamMemory.HeaderSizes.GetData());
+
+			// Describe Data block
+			CachedAPI->rmx_output_media_set_packet_layout(&Block, StreamMemory.DataBlockID, StreamMemory.PayloadSizes.GetData());
+
+			rmx_mem_multi_key_region* DataMemory;
+			DataMemory = CachedAPI->rmx_output_media_get_dup_sub_block(&Block, StreamMemory.DataBlockID);
+			if (DataMemory == nullptr)
+			{
+				UE_LOG(LogRivermax, Warning, TEXT("Faild to get payload memory block. Output stream won't be created."));
+				return false;
+			}
+
+			// If intermediate buffer is used, we setup rmax memblock to use that address. Otherwise, we map it to our actual frame's address
+			if (StreamMemory.bUseIntermediateBuffer)
+			{
+				DataMemory->addr = Allocator->GetFrameAddress(BlockIndex);
+			}
+			else
+			{
+				TSharedPtr<UE::RivermaxCore::Private::FRivermaxOutputFrame> ProcessedFrame = FrameManager->GetFreeFrame();
+				DataMemory->addr = ProcessedFrame->Buffer;
+				ProcessedFrames.Add(ProcessedFrame);
+			}
+
+			DataMemory->length = StreamMemory.ChunksPerFrameField * StreamMemory.PacketsPerChunk * StreamMemory.PayloadSize;
+
+			constexpr rmx_mkey_id InvalidKey = ((rmx_mkey_id)(-1L));
+			DataMemory->mkey[0] = InvalidKey;
+			DataMemory->mkey[1] = InvalidKey;
+
+
+			rmx_mem_multi_key_region* HeaderMemory;
+			HeaderMemory = CachedAPI->rmx_output_media_get_dup_sub_block(&Block, StreamMemory.HeaderBlockID);
+			if (HeaderMemory == nullptr)
+			{
+				UE_LOG(LogRivermax, Warning, TEXT("Faild to get header memory block. Output stream won't be created."));
+				return false;
+			}
+
+			StreamMemory.RTPHeaders[BlockIndex].SetNumZeroed(RealPacketsPerFrame * StreamMemory.FramesFieldPerMemoryBlock);
+			HeaderMemory->addr = &StreamMemory.RTPHeaders[BlockIndex][0];
+			HeaderMemory->length = StreamMemory.HeaderStrideSize;
+			HeaderMemory->mkey[0] = InvalidKey;
+			HeaderMemory->mkey[1] = InvalidKey;
+		}
+		uint64 ProcessedBytes = 0;
+		uint64 LineSize = 0;
+		for (int32 PayloadSizeIndex = 0; PayloadSizeIndex < RealPacketsPerFrame; ++PayloadSizeIndex)
+		{
+			uint32 HeaderSize = FVideoRTPHeader::OneSRDSize;
+			uint32 ThisPayloadSize = StreamMemory.PayloadSize;
+			if (ProcessedBytes < FrameSizeInBytes)
+			{
+				if ((LineSize + StreamMemory.PayloadSize) == BytesPerRow)
+				{
+					LineSize = 0;
+				}
+				else if ((LineSize + StreamMemory.PayloadSize) > BytesPerRow)
+				{
+					HeaderSize = FVideoRTPHeader::TwoSRDSize;
+					LineSize = StreamMemory.PayloadSize - (BytesPerRow - LineSize);
+					if (LineSize > BytesPerRow)
+					{
+						UE_LOG(LogRivermax, Warning, TEXT("Unsupported small resolution, %dx%d, needing more than 2 SRD to express"), StreamOptions->AlignedResolution.X, StreamOptions->AlignedResolution.Y);
+						return false;
+					}
+				}
+				else
+				{
+					// Keep track of line size offset to know when to use TwoSRDs
+					LineSize += StreamMemory.PayloadSize;
+				}
+
+				// This means that this is the header for the last packet and it is smaller than all the previous ones.
+				if ((ProcessedBytes + StreamMemory.PayloadSize) > FrameSizeInBytes)
+				{
+					HeaderSize = FVideoRTPHeader::OneSRDSize;
+					ThisPayloadSize = FrameSizeInBytes - ProcessedBytes;
+				}
+			}
+			else
+			{
+				// Extra header/payload required for the chunk alignment are set to 0. Nothing has to be sent out the wire.
+				HeaderSize = 0;
+				ThisPayloadSize = 0;
+			}
+
+			// All buffers are configured the same so compute header and payload sizes once and assigned to all impacted locations
+			for (uint32 BufferIndex = 0; BufferIndex < StreamMemory.FramesFieldPerMemoryBlock; ++BufferIndex)
+			{
+				StreamMemory.HeaderSizes[PayloadSizeIndex + (BufferIndex * RealPacketsPerFrame)] = HeaderSize;
+				StreamMemory.PayloadSizes[PayloadSizeIndex + (BufferIndex * RealPacketsPerFrame)] = ThisPayloadSize;
+			}
+
+			if (CachedCVars.bPrefillRTPHeaders && HeaderSize > 0)
+			{
+				RTPFiller.Update(PayloadSizeIndex);
+			}
+
+			ProcessedBytes += ThisPayloadSize;
+		}
+
+		// Verify memcopy config to make sure it works for current frame size / chunking
+		if (StreamMemory.bUseIntermediateBuffer)
+		{
+			StreamMemory.FrameMemorySliceCount = FMath::Clamp(CVarRivermaxOutputMaxFrameMemorySliceCount.GetValueOnAnyThread(), 1, 100);
+			StreamMemory.ChunkSpacingBetweenMemcopies = FMath::Clamp(CVarRivermaxOutputMemcopyChunkSpacing.GetValueOnAnyThread(), 1, 20);
+
+			const uint32 ChunkRequired = StreamMemory.ChunkSpacingBetweenMemcopies * StreamMemory.FrameMemorySliceCount;
+			if (ChunkRequired > 0 && ChunkRequired > StreamMemory.ChunksPerFrameField)
+			{
+				// Favor reducing number of memcopies. If required packet is smaller, chances are it's a small frame size
+				// so memcopies will be smaller.
+				const double Ratio = StreamMemory.ChunksPerFrameField / (double)ChunkRequired;
+				StreamMemory.FrameMemorySliceCount *= Ratio;
+			}
+		}
+
+		return true;
+	}
+
+	bool FRivermaxOutVideoStream::IsFrameAvailableToSend()
+	{
+		return FrameManager->IsFrameAvailableToSend();
+	}
+
+	bool FRivermaxOutVideoStream::CopyFrameData(const TSharedPtr<FRivermaxOutputFrame>& SourceFrame, uint8* DestinationBase)
+	{
+		// Make sure copy size doesn't go over frame size
+		const SIZE_T FrameSize = GetRowSizeInBytes() * Options.GetStreamOptions<FRivermaxVideoOutputOptions>(StreamType)->AlignedResolution.Y;
+		if (ensure(FrameSize > 0))
+		{
+			const SIZE_T BlockSize = 1 + ((FrameSize - 1) / StreamMemory.FrameMemorySliceCount);
+			const SIZE_T MaxSize = FrameSize - SourceFrame->Offset;
+			const SIZE_T CopySize = FMath::Min(BlockSize, MaxSize);
+
+			// Copy data until we have covered the whole frame. Last block might be smaller.
+			if (CopySize > 0)
+			{
+				TRACE_CPUPROFILER_EVENT_SCOPE(CopyFrameData);
+				const SIZE_T PacketsToCopy = CopySize / StreamMemory.PayloadSize;
+				const SIZE_T PacketOffset = SourceFrame->ChunkNumber * StreamMemory.PacketsPerChunk;
+				const uint64 PacketDeltaTime = (StreamData.FrameFieldTimeIntervalNs - TransmitOffsetNanosec) / StreamMemory.PacketsPerFrame;
+
+				FCopyArgs Args;
+				uint8* SourceStart = reinterpret_cast<uint8*>(SourceFrame->Buffer);
+				uint8* DestinationStart = DestinationBase;
+				Args.SourceMemory = SourceStart + SourceFrame->Offset;
+				Args.DestinationMemory = DestinationStart + SourceFrame->Offset;
+				Args.SizeToCopy = CopySize;
+
+				// Update memory offset for next copy
+				SourceFrame->Offset += CopySize;
+
+				Allocator->CopyData(Args);
+
+				return true;
+			}
+		}
+		else
+		{
+			UE_LOG(LogRivermax, Error, TEXT("Invalid frame size detected while stream was active. Shutting down."));
+			Listener->OnStreamError();
+			Stop();
+		}
+
+		return false;
+	}
+
+	void FRivermaxOutVideoStream::InitializeStreamTimingSettings()
+	{
+		FRivermaxOutStream::InitializeStreamTimingSettings();
+		using namespace UE::RivermaxCore::Private::Utils;
+
+		double FrameIntervalNs = StreamData.FrameFieldTimeIntervalNs;
+		const bool bIsProgressive = true;//todo MediaConfiguration.IsProgressive() 
+		uint32 PacketsInFrameField = StreamMemory.PacketsPerFrame;
+		if (bIsProgressive == false)
+		{
+			FrameIntervalNs *= 2;
+			PacketsInFrameField *= 2;
+		}
+
+		if (StreamType == ERivermaxStreamType::VIDEO_2110_20_STREAM)
+		{
+			double RActive;
+			double TRODefaultMultiplier;
+			TSharedPtr<FRivermaxVideoOutputOptions> VideoOptions = Options.GetStreamOptions<FRivermaxVideoOutputOptions>(StreamType);
+
+			// See https://ieeexplore.ieee.org/stamp/stamp.jsp?tp=&arnumber=8165971 for reference
+			// Gapped PRS doesn't support non standard resolution. Linear PRS would but Rivermax doesn't support it.
+			if (bIsProgressive)
+			{
+				RActive = (1080.0 / 1125.0);
+				if (VideoOptions->AlignedResolution.Y >= FullHDHeight)
+				{
+					// As defined by SMPTE 2110-21 6.3.2
+					TRODefaultMultiplier = (43.0 / 1125.0);
+				}
+				else
+				{
+					TRODefaultMultiplier = (28.0 / 750.0);
+				}
+			}
+			else
+			{
+				if (VideoOptions->AlignedResolution.Y >= FullHDHeight)
+				{
+					// As defined by SMPTE 2110-21 6.3.3
+					RActive = (1080.0 / 1125.0);
+					TRODefaultMultiplier = (22.0 / 1125.0);
+				}
+				else if (VideoOptions->AlignedResolution.Y >= 576)
+				{
+					RActive = (576.0 / 625.0);
+					TRODefaultMultiplier = (26.0 / 625.0);
+				}
+				else
+				{
+					RActive = (487.0 / 525.0);
+					TRODefaultMultiplier = (20.0 / 525.0);
+				}
+			}
+
+			// Need to reinvestigate the implication of this and possibly add cvar to control it runtime
+			const double TRSNano = (FrameIntervalNs * RActive) / PacketsInFrameField;
+			TransmitOffsetNanosec = (uint64)((TRODefaultMultiplier * FrameIntervalNs));
+		}
+	}
+
+	bool FRivermaxOutVideoStream::SetupFrameManagement()
+	{
+		TRACE_CPUPROFILER_EVENT_SCOPE(FRivermaxOutVideoStream::SetupFrameManagement);
+
+		FrameManager = MakeUnique<FFrameManager>();
+		TSharedPtr<FRivermaxVideoOutputOptions> VideoOptions = Options.GetStreamOptions<FRivermaxVideoOutputOptions>(StreamType);
+
+		// We do (try) to make gpu allocations here to let the capturer know if we require it or not.
+		bool bTryGPUDirect = RivermaxModule->GetRivermaxManager()->IsGPUDirectOutputSupported() && VideoOptions->bUseGPUDirect;
+		if (bTryGPUDirect)
+		{
+			const ERHIInterfaceType RHIType = RHIGetInterfaceType();
+			if (RHIType != ERHIInterfaceType::D3D12)
+			{
+				UE_LOG(LogRivermax, Warning, TEXT("Can't initialize output to use GPUDirect. RHI is %d but only Dx12 is supported at the moment."), RHIType);
+				bTryGPUDirect = false;
+			}
+		}
+
+		// Work around when dealing with multi memblocks. Rivermax fails to create stream 
+		// With memblock not starting on the right cuda alignment. 
+		const bool bAlignEachFrameMemory = !CachedCVars.bUseSingleMemblock;
+		FFrameManagerSetupArgs FrameManagerArgs;
+		FrameManagerArgs.Resolution = VideoOptions->AlignedResolution;
+		FrameManagerArgs.bTryGPUAllocation = bTryGPUDirect;
+		FrameManagerArgs.NumberOfFrames = Options.NumberOfBuffers;
+		FrameManagerArgs.Stride = GetRowSizeInBytes();
+		FrameManagerArgs.FrameDesiredSize = StreamMemory.ChunksPerFrameField * StreamMemory.PacketsPerChunk * StreamMemory.PayloadSize;
+		FrameManagerArgs.bAlignEachFrameAlloc = bAlignEachFrameMemory;
+		FrameManagerArgs.OnFreeFrameDelegate = FOnFrameReadyDelegate::CreateRaw(this, &FRivermaxOutVideoStream::OnFrameReadyToBeUsed);
+		FrameManagerArgs.OnPreFrameReadyDelegate = FOnFrameReadyDelegate::CreateRaw(this, &FRivermaxOutVideoStream::OnPreFrameReadyToBeSent);
+		FrameManagerArgs.OnFrameReadyDelegate = FOnFrameReadyDelegate::CreateRaw(this, &FRivermaxOutVideoStream::OnFrameReadyToBeSent);
+		FrameManagerArgs.OnCriticalErrorDelegate = FOnCriticalErrorDelegate::CreateRaw(this, &FRivermaxOutVideoStream::OnFrameManagerCriticalError);
+		const EFrameMemoryLocation FrameLocation = FrameManager->Initialize(FrameManagerArgs);
+		bUseGPUDirect = FrameLocation == EFrameMemoryLocation::GPU;
+
+
+		// Only support intermediate buffer for alignment point method to avoid running into chunk issue when repeating a frame
+		const bool bHasAllocatedFrames = (FrameLocation != EFrameMemoryLocation::None);
+		if (bHasAllocatedFrames && Options.AlignmentMode == ERivermaxAlignmentMode::AlignmentPoint && CVarRivermaxOutputEnableIntermediateBuffer.GetValueOnAnyThread())
+		{
+			// Allocate intermediate buffer in same memory space as frame memory.
+			FOnFrameDataCopiedDelegate OnDataCopiedDelegate = FOnFrameDataCopiedDelegate::CreateRaw(this, &FRivermaxOutVideoStream::OnMemoryChunksCopied);
+
+			const int32 DesiredSize = FrameManagerArgs.FrameDesiredSize * StreamMemory.FramesFieldPerMemoryBlock;
+			if (bUseGPUDirect)
+			{
+				Allocator = MakeUnique<FGPUAllocator>(DesiredSize, OnDataCopiedDelegate);
+			}
+			else
+			{
+				Allocator = MakeUnique<FSystemAllocator>(DesiredSize, OnDataCopiedDelegate);
+			}
+
+			if (!Allocator->Allocate(StreamMemory.MemoryBlockCount, bAlignEachFrameMemory))
+			{
+				return false;
+			}
+
+			StreamMemory.bUseIntermediateBuffer = true;
+		}
+
+		// Cache buffer addresses used by rivermax in order to start copying into it early on
+		StreamMemory.BufferAddresses.Reserve(Options.NumberOfBuffers);
+		if (StreamMemory.bUseIntermediateBuffer)
+		{
+			check(Allocator);
+
+			for (uint32 MemblockIndex = 0; MemblockIndex < StreamMemory.MemoryBlockCount; ++MemblockIndex)
+			{
+				uint8* BaseAddress = static_cast<uint8*>(Allocator->GetFrameAddress(MemblockIndex));
+				for (uint32 FrameIndex = 0; FrameIndex < StreamMemory.FramesFieldPerMemoryBlock; ++FrameIndex)
+				{
+					uint8* FrameAddress = BaseAddress + (FrameIndex * FrameManagerArgs.FrameDesiredSize);
+					StreamMemory.BufferAddresses.Add(FrameAddress);
+				}
+			}
+		}
+		else
+		{
+			// Used to keep track of acquired frames to prevent the same frame from being acquired from the pool more than once. 
+			// When this scope is exited all frames are returned back to the Frame Manager pool.
+			TArray<TSharedPtr<UE::RivermaxCore::Private::FRivermaxOutputFrame>> ProcessedFrames;
+
+			// When we don't use intermediate buffer, each frame has its own buffer addresses and we don't need to look at memblocks
+			for (int32 BufferIndex = 0; BufferIndex < Options.NumberOfBuffers; ++BufferIndex)
+			{
+				TSharedPtr<UE::RivermaxCore::Private::FRivermaxOutputFrame> Frame = FrameManager->GetFreeFrame();
+				StreamMemory.BufferAddresses.Add(Frame->Buffer);
+				ProcessedFrames.Add(Frame);
+			}
+		}
+		check(StreamMemory.BufferAddresses.Num() == Options.NumberOfBuffers);
+
+		return FrameLocation != EFrameMemoryLocation::None;
+	}
+
+	void FRivermaxOutVideoStream::CleanupFrameManagement()
+	{
+		if (FrameManager)
+		{
+			FrameManager->Cleanup();
+			FrameManager.Reset();
+		}
+
+		if (Allocator)
+		{
+			Allocator->Deallocate();
+			Allocator.Reset();
+		}
+	}
+
+	SIZE_T FRivermaxOutVideoStream::GetRowSizeInBytes() const
+	{
+		check(FormatInfo.PixelGroupCoverage != 0);
+		TSharedPtr<FRivermaxVideoOutputOptions> VideoOptions = StaticCastSharedPtr<FRivermaxVideoOutputOptions>(Options.StreamOptions[(uint8)StreamType]);
+		return (VideoOptions->AlignedResolution.X / FormatInfo.PixelGroupCoverage) * FormatInfo.PixelGroupSize;
+	}
+
+	void FRivermaxOutVideoStream::SetupRTPHeaders()
+	{
+		FVideoRTPHeader* FirstHeaderPtr = reinterpret_cast<FVideoRTPHeader*>(CurrentFrame->HeaderPtr);
+		check(FirstHeaderPtr);
+
+		for (uint32 PacketIndex = 0; PacketIndex < StreamMemory.PacketsPerChunk && CurrentFrame->PacketCounter < StreamMemory.PacketsPerFrame; ++PacketIndex)
+		{
+			using namespace UE::RivermaxCore::Private::Utils;
+			FVideoRTPHeader* HeaderRawPtr = FirstHeaderPtr+ PacketIndex;
+			// Since the static part of the RTP header is filled on initalization we only need to update the 
+			// non static parts.
+			if (CachedCVars.bPrefillRTPHeaders)
+			{
+				HeaderRawPtr->RTPHeader.SequenceNumber = ByteSwap((uint16)(StreamData.SequenceNumber & 0xFFFF));
+				HeaderRawPtr->RTPHeader.Timestamp = ByteSwap(CurrentFrame->MediaTimestamp);
+				HeaderRawPtr->RTPHeader.ExtendedSequenceNumber = ByteSwap((uint16)((StreamData.SequenceNumber >> 16) & 0xFFFF));
+			}
+			else
+			{
+				TSharedPtr<FRivermaxVideoOutputOptions> VideoOptions = StaticCastSharedPtr<FRivermaxVideoOutputOptions>(Options.StreamOptions[(uint8)StreamType]);
+				UpdateVideoRtpHeader(
+					*HeaderRawPtr,
+					CurrentFrame->PacketCounter,
+					CurrentFrame->SRDOffset,
+					CurrentFrame->LineNumber,
+					StreamData.SequenceNumber,
+					CurrentFrame->MediaTimestamp,
+					VideoOptions->AlignedResolution,
+					FormatInfo,
+					StreamMemory
+				);
+			}
+
+			++StreamData.SequenceNumber;
+			++CurrentFrame->PacketCounter;
+		}
+	}
+
+	void FRivermaxOutVideoStream::LogStreamDescriptionOnCreation() const
+	{
+		FRivermaxOutStream::LogStreamDescriptionOnCreation();
+
+		TSharedPtr<FRivermaxVideoOutputOptions> VideoOptions = StaticCastSharedPtr<FRivermaxVideoOutputOptions>(Options.StreamOptions[(uint8)StreamType]);
+
+		TStringBuilder<512> StreamDescription;
+		if (bUseGPUDirect)
+		{
+			StreamDescription.Appendf(TEXT("Using GPUDirect."));
+		}
+
+		StreamDescription.Appendf(TEXT("Settings: Resolution = %dx%d, "), VideoOptions->AlignedResolution.X, VideoOptions->AlignedResolution.Y);
+		StreamDescription.Appendf(TEXT("FrameRate = %s, "), *VideoOptions->FrameRate.ToPrettyText().ToString());
+		StreamDescription.Appendf(TEXT("Pixel format = %s, "), LexToString(VideoOptions->PixelFormat));
+		StreamDescription.Appendf(TEXT("Alignment = %s, "), LexToString(Options.AlignmentMode));
+		StreamDescription.Appendf(TEXT("Framelocking = %s."), LexToString(Options.FrameLockingMode));
+
+		UE_LOG(LogRivermax, Display, TEXT("%s"), *FString(StreamDescription));
+	}
+
+	bool FRivermaxOutVideoStream::ReserveFrame(uint64 FrameCounter) const
+	{
+		// There is only one reserved frame at the time per stream.
+		TSharedPtr<FRivermaxOutputFrame> ReservedFrame = FrameManager->GetFreeFrame();
+		if (!ReservedFrame.IsValid() && Options.FrameLockingMode == EFrameLockingMode::BlockOnReservation)
+		{
+			TRACE_CPUPROFILER_EVENT_SCOPE(Rmax::WaitForAvailableFrame);
+			while (!ReservedFrame && bIsActive)
+			{
+				FrameAvailableSignal->Wait();
+				ReservedFrame = FrameManager->GetFreeFrame();
+			}
+		}
+
+		if (ReservedFrame.IsValid())
+		{
+			ReservedFrame->SetFrameCounter(FrameCounter);
+			ReservedFrames.Enqueue(ReservedFrame);
+		}
+
+		return ReservedFrame.IsValid();
+	}
+
+	TSharedPtr<FRivermaxOutputFrame> FRivermaxOutVideoStream::GetNextFrameToSend(bool bWait)
+	{
+		TSharedPtr<FRivermaxOutputFrame> NextFrameToSend = FrameManager->DequeueFrameToSend();
+
+		if (bWait)
+		{
+			while (!NextFrameToSend.IsValid() && bIsActive)
+			{
+				FrameReadyToSendSignal->Wait();
+				NextFrameToSend = FrameManager->DequeueFrameToSend();
+			}
+		}
+		return NextFrameToSend;
+	}
+
+	void FRivermaxOutVideoStream::CompleteCurrentFrame(bool bReleaseFrame)
+	{
+		FRivermaxOutStream::CompleteCurrentFrame(bReleaseFrame);
+
+		// We don't release when there is no new frame, so we keep a hold on it to repeat it.
+		if (bReleaseFrame)
+		{
+			FrameManager->FrameSentEvent();
+		}
+	}
+}
+
diff --git a/Engine/Plugins/VirtualProduction/Rivermax/RivermaxCore/Source/RivermaxCore/Private/Streams/RivermaxOutVideoStream.h b/Engine/Plugins/VirtualProduction/Rivermax/RivermaxCore/Source/RivermaxCore/Private/Streams/RivermaxOutVideoStream.h
new file mode 100644
--- /dev/null
+++ b/Engine/Plugins/VirtualProduction/Rivermax/RivermaxCore/Source/RivermaxCore/Private/Streams/RivermaxOutVideoStream.h
@@ -0,0 +1,74 @@
+// Copyright Epic Games, Inc. All Rights Reserved.
+
+#pragma once
+
+#include "RivermaxOutStream.h"
+
+#include "Async/Future.h"
+#include "Containers/SpscQueue.h"
+#include "HAL/Runnable.h"
+#include "RivermaxWrapper.h"
+#include "RivermaxOutputFrame.h"
+#include "RivermaxTypes.h"
+#include "RTPHeader.h"
+
+
+class FEvent;
+class IRivermaxCoreModule;
+
+namespace UE::RivermaxCore::Private
+{
+	class FFrameManager;
+	class FBaseFrameAllocator;
+
+	using UE::RivermaxCore::FRivermaxOutputOptions;
+
+	class FRivermaxOutVideoStream : public FRivermaxOutStream
+	{
+	public:
+		FRivermaxOutVideoStream(const TArray<char>& SDPDescription);
+
+	public:
+
+		//~ Begin IRivermaxOutputStream interface
+		virtual bool PushFrame(TSharedPtr<IRivermaxOutputInfo> FrameInfo) override;
+		virtual bool ReserveFrame(uint64 FrameIdentifier) const override;
+		//~ End IRivermaxOutputStream interface
+
+	protected:
+
+		//~ Begin FRivermaxOutStream interface
+		virtual bool IsFrameAvailableToSend();
+		virtual bool InitializeStreamMemoryConfig();
+		virtual bool CopyFrameData(const TSharedPtr<FRivermaxOutputFrame>& SourceFrame, uint8* DestinationBase);
+		virtual bool SetupFrameManagement();
+		virtual void SetupRTPHeaders() override;
+		virtual void CleanupFrameManagement();
+		virtual void InitializeStreamTimingSettings() override;
+		virtual void CompleteCurrentFrame(bool bReleaseFrame) override;
+		virtual void LogStreamDescriptionOnCreation() const override;
+		virtual TSharedPtr<FRivermaxOutputFrame> GetNextFrameToSend(bool bWait = false);
+		//~ End FRivermaxOutStream interface
+
+	private:
+		
+		/** Get row stride for the current stream configuration */
+		SIZE_T GetRowSizeInBytes() const;
+
+	private:
+
+		/** Manages allocation and memory manipulation of video frames */
+		TUniquePtr<FFrameManager> FrameManager;
+
+		/** Manages allocation of memory for rivermax memblocks */
+		TUniquePtr<FBaseFrameAllocator> Allocator;
+
+		/** Format info for the active stream */
+		FVideoFormatInfo FormatInfo;
+
+		// A helper struct that is hidden in cpp file and needs access to video stream options to fill the RTP packet.
+		friend struct FRTPHeaderPrefiller;
+	};
+}
+
+
diff --git a/Engine/Plugins/VirtualProduction/Rivermax/RivermaxCore/Source/RivermaxCore/Public/IRivermaxCoreModule.h b/Engine/Plugins/VirtualProduction/Rivermax/RivermaxCore/Source/RivermaxCore/Public/IRivermaxCoreModule.h
--- a/Engine/Plugins/VirtualProduction/Rivermax/RivermaxCore/Source/RivermaxCore/Public/IRivermaxCoreModule.h
+++ b/Engine/Plugins/VirtualProduction/Rivermax/RivermaxCore/Source/RivermaxCore/Public/IRivermaxCoreModule.h
@@ -7,6 +7,7 @@
 #include "IRivermaxInputStream.h"
 #include "IRivermaxOutputStream.h"
 #include "Modules/ModuleManager.h"
+#include "RivermaxTypes.h"
 
 
 namespace UE::RivermaxCore
@@ -33,11 +34,17 @@ public:
 		return FModuleManager::Get().IsModuleLoaded("RivermaxCore");
 	}
 
+	UE_DEPRECATED(5.6, "Please use CreateInputStream that takes stream type as a parameter.")
+	virtual TUniquePtr<UE::RivermaxCore::IRivermaxInputStream> CreateInputStream() { return nullptr; };
+
+	UE_DEPRECATED(5.6, "Please use CreateOutputStream that takes stream type as a parameter.")
+	virtual TUniquePtr<UE::RivermaxCore::IRivermaxOutputStream> CreateOutputStream() { return nullptr; };
+
 	/** Create input stream managing receiving data from rivermax */
-	virtual TUniquePtr<UE::RivermaxCore::IRivermaxInputStream> CreateInputStream() = 0;
+	virtual TUniquePtr<UE::RivermaxCore::IRivermaxInputStream> CreateInputStream(UE::RivermaxCore::ERivermaxStreamType, const TArray<char>& InSDPDescription) = 0;
 
 	/** Create output stream managing sending data to rivermax */
-	virtual TUniquePtr<UE::RivermaxCore::IRivermaxOutputStream> CreateOutputStream() = 0;
+	virtual TUniquePtr<UE::RivermaxCore::IRivermaxOutputStream> CreateOutputStream(UE::RivermaxCore::ERivermaxStreamType, const TArray<char>& InSDPDescription) = 0;
 
 	/** Returns Rivermax manager singleton to query for stats, status, etc... */
 	virtual TSharedPtr<UE::RivermaxCore::IRivermaxManager> GetRivermaxManager() = 0;
diff --git a/Engine/Plugins/VirtualProduction/Rivermax/RivermaxCore/Source/RivermaxCore/Public/IRivermaxOutputStream.h b/Engine/Plugins/VirtualProduction/Rivermax/RivermaxCore/Source/RivermaxCore/Public/IRivermaxOutputStream.h
--- a/Engine/Plugins/VirtualProduction/Rivermax/RivermaxCore/Source/RivermaxCore/Public/IRivermaxOutputStream.h
+++ b/Engine/Plugins/VirtualProduction/Rivermax/RivermaxCore/Source/RivermaxCore/Public/IRivermaxOutputStream.h
@@ -9,23 +9,27 @@
 
 namespace UE::RivermaxCore
 {
-	struct FRivermaxOutputStreamOptions;
+	struct FRivermaxOutputOptions;
 
 	/** Description of video frame to be captured. */
-	struct RIVERMAXCORE_API FRivermaxOutputVideoFrameInfo
+	struct RIVERMAXCORE_API IRivermaxOutputInfo
 	{
 		/** Incremental number identifying frame. Usually GFrameCounter */
-		uint32 FrameIdentifier = 0;
+		uint64 FrameIdentifier = 0;
 
 		/** Height of the frame */
 		uint32 Height = 0;
-		
+
 		/** Width of the frame */
 		uint32 Width = 0;
 
 		/** Stride of a line */
 		uint32 Stride = 0;
+	};
 
+	/** Description of video frame to be captured. */
+	struct RIVERMAXCORE_API FRivermaxOutputVideoFrameInfo : public IRivermaxOutputInfo
+	{
 		/** Video frame data location in system memory */
 		void* CPUBuffer = nullptr;
 
@@ -33,6 +37,12 @@ namespace UE::RivermaxCore
 		FBufferRHIRef GPUBuffer;
 	};
 
+	/** Description of video frame to be captured. */
+	struct RIVERMAXCORE_API FRivermaxOutputAncInfo : public IRivermaxOutputInfo
+	{
+
+	};
+
 	/** Information about the last frame that was presented by the stream */
 	struct RIVERMAXCORE_API FPresentedFrameInfo
 	{
@@ -68,31 +78,27 @@ namespace UE::RivermaxCore
 		 * Initializes stream using input options. Returns false if stream creation has failed. 
 		 * Initialization is completed when OnInitializationCompletion has been called with result.
 		 */
-		virtual bool Initialize(const FRivermaxOutputStreamOptions& Options, IRivermaxOutputStreamListener& InListener) = 0;
+		virtual bool Initialize(const FRivermaxOutputOptions& Options, IRivermaxOutputStreamListener& InListener) = 0;
 
 		/** Uninitializes current stream */
 		virtual void Uninitialize() = 0;
 
 		/** Pushes new video frame to the stream. Returns false if frame couldn't be pushed. */
-		virtual bool PushVideoFrame(const FRivermaxOutputVideoFrameInfo& NewFrame) = 0;
+		UE_DEPRECATED(5.6, "This method is deprecated. Please use the PushFrame.")
+		virtual bool PushVideoFrame(const FRivermaxOutputVideoFrameInfo& NewFrame) { return false; };
+
+		/** Pushes new video frame to the stream. Returns false if frame couldn't be pushed. */
+		virtual bool PushFrame(TSharedPtr<IRivermaxOutputInfo> FrameInfo) = 0;
 
 		/** Returns true if GPUDirect is supported */
 		virtual bool IsGPUDirectSupported() const = 0;
 
 		/** 
 		 * Tries to reserve a frame for the next capture 
-		 * If FRivermaxOutputStreamOptions::EFrameLockingMode equals BlockOnReservation, 
-		 * this method will block until a free frame is found. 
-		 */
-		UE_DEPRECATED(5.6, "This method is deprecated. Please use the ReserveFrame that doesn't rely on frame identifier.")
-		virtual bool ReserveFrame(uint32 FrameIdentifier) const { return false; };
-
-		/** 
-		 * Tries to reserve a frame for the next capture 
-		 * If FRivermaxOutputStreamOptions::EFrameLockingMode equals BlockOnReservation, 
+		 * If FRivermaxOutputOptions::FrameLockingMode equals BlockOnReservation, 
 		 * this method will block until a free frame is found. 
 		 */
-		virtual bool ReserveFrame() = 0;
+		virtual bool ReserveFrame(uint64 FrameCounter) const = 0;
 
 		/** Returns information about the last frame that was presented on the wire */
 		virtual void GetLastPresentedFrame(FPresentedFrameInfo& OutFrameInfo) const = 0;
diff --git a/Engine/Plugins/VirtualProduction/Rivermax/RivermaxCore/Source/RivermaxCore/Public/RivermaxTypes.h b/Engine/Plugins/VirtualProduction/Rivermax/RivermaxCore/Source/RivermaxCore/Public/RivermaxTypes.h
--- a/Engine/Plugins/VirtualProduction/Rivermax/RivermaxCore/Source/RivermaxCore/Public/RivermaxTypes.h
+++ b/Engine/Plugins/VirtualProduction/Rivermax/RivermaxCore/Source/RivermaxCore/Public/RivermaxTypes.h
@@ -7,6 +7,7 @@
 #include "Misc/FrameRate.h"
 #include "RivermaxFormats.h"
 
+
 namespace UE::RivermaxCore
 {
 	constexpr TCHAR DefaultStreamAddress[] = TEXT("228.1.1.1");
@@ -104,9 +105,6 @@ namespace UE::RivermaxCore
 
 	struct FRivermaxOutputStreamOptions
 	{
-		/** Desired stream resolution */
-		FIntPoint Resolution = { 1920, 1080 };
-
 		/** Stream FrameRate */
 		FFrameRate FrameRate = { 24,1 };
 
@@ -119,17 +117,54 @@ namespace UE::RivermaxCore
 		/** Port to be used by stream */
 		uint32 Port = 50000;
 
+		/** Used by RivermaxOutStream when it calls to the library to assign Media Block Index in SDP. */
+		uint64 StreamIndex = 0;
+	};
+
+	struct FRivermaxVideoOutputOptions : public FRivermaxOutputStreamOptions
+	{
+		/** Desired stream resolution */
+		FIntPoint Resolution = { 1920, 1080 };
+
 		/** Desired stream pixel format */
 		ESamplingType PixelFormat = ESamplingType::RGB_10bit;
 
-		/** Sample count to buffer. */
-		int32 NumberOfBuffers = 2;
-
 		/** Resolution aligning with pgroup of sampling type */
 		FIntPoint AlignedResolution = FIntPoint::ZeroValue;
 
 		/** Whether to leverage GPUDirect (Cuda) capability to transfer memory to NIC if available */
 		bool bUseGPUDirect = true;
+	};
+
+	struct FRivermaxAncOutputOptions : public FRivermaxOutputStreamOptions
+	{
+		/** DID value is specified by SMPTE 291 standard. 0x60 - Ancillary timecode. */
+		uint16 DID = 0x60;
+
+		/** SDID value is specified by SMPTE 291 standard. 0x60 - Ancillary timecode. */
+		uint16 SDID = 0x60;
+
+		FRivermaxAncOutputOptions()
+		{
+			Port = 50010;
+		};
+	};
+
+	enum class ERivermaxStreamType : uint8
+	{
+		VIDEO_2110_20_STREAM,
+		AUDIO_2110_30_STREAM,
+		ANC_2110_40_STREAM,
+		MAX
+	};
+
+	struct FRivermaxOutputOptions
+	{
+		/** If not nullptr indicates that anc stream needs to be created. Contains all anc stream related options. */
+		TStaticArray<TSharedPtr<FRivermaxOutputStreamOptions>, static_cast<uint32>(ERivermaxStreamType::MAX)> StreamOptions;
+
+		/** Sample count to buffer. */
+		int32 NumberOfBuffers = 2;
 
 		/** Method used to align output stream */
 		ERivermaxAlignmentMode AlignmentMode = ERivermaxAlignmentMode::AlignmentPoint;
@@ -142,11 +177,18 @@ namespace UE::RivermaxCore
 
 		/** Whether to use frame's frame number instead of standard timestamping */
 		bool bDoFrameCounterTimestamping = true;
-	};
 
-	enum class ERivermaxStreamType : uint8
-	{
-		VIDEO_2110_20_STREAM,
+		/** Returns stream options for this stream. */
+		template <typename T>
+		TSharedPtr<T> GetStreamOptions(ERivermaxStreamType InStreamType)
+		{
+			return StaticCastSharedPtr<T>(StreamOptions[static_cast<uint8>(InStreamType)]);
+		}
+
+		TSharedPtr<FRivermaxOutputStreamOptions> GetStreamOptions(ERivermaxStreamType InStreamType)
+		{
+			return StreamOptions[static_cast<uint8>(InStreamType)];
+		}
 	};
 }
 
diff --git a/Engine/Plugins/VirtualProduction/Rivermax/RivermaxMedia/Source/RivermaxMedia/Private/RivermaxMediaCapture.cpp b/Engine/Plugins/VirtualProduction/Rivermax/RivermaxMedia/Source/RivermaxMedia/Private/RivermaxMediaCapture.cpp
--- a/Engine/Plugins/VirtualProduction/Rivermax/RivermaxMedia/Source/RivermaxMedia/Private/RivermaxMediaCapture.cpp
+++ b/Engine/Plugins/VirtualProduction/Rivermax/RivermaxMedia/Source/RivermaxMedia/Private/RivermaxMediaCapture.cpp
@@ -55,15 +55,6 @@ namespace RivermaxMediaCaptureAnalytics
 
 namespace UE::RivermaxMediaCaptureUtil
 {
-	FIntPoint GetAlignedResolution(const UE::RivermaxCore::FVideoFormatInfo& InFormatInfo, const FIntPoint& InSize)
-	{
-		// We need to adjust horizontal resolution if it's not aligned to pgroup pixel coverage
-		// RTP header has to describe a full pgroup even if it's outside of effective resolution.
-		const uint32 PixelAlignment = InFormatInfo.PixelGroupCoverage;
-		const uint32 AlignedHorizontalResolution = (InSize.X % PixelAlignment) ? InSize.X + (PixelAlignment - (InSize.X % PixelAlignment)) : InSize.X;
-		return FIntPoint(AlignedHorizontalResolution, InSize.Y);
-	}
-
 	void GetOutputEncodingInfo(ERivermaxMediaOutputPixelFormat InPixelFormat, const FIntPoint& InSize, uint32& OutBytesPerElement, uint32& OutElementsPerFrame)
 	{
 		using namespace UE::RivermaxCore;
@@ -73,7 +64,7 @@ namespace UE::RivermaxMediaCaptureUtil
 
 		// Compute horizontal byte count (stride) of aligned resolution
 		const uint32 PixelAlignment = Info.PixelGroupCoverage;
-		const FIntPoint AlignedResolution = GetAlignedResolution(Info, InSize);
+		const FIntPoint AlignedResolution = UE::RivermaxMediaUtils::Private::GetAlignedResolution(Info, InSize);
 		const uint32 PixelCount = AlignedResolution.X * AlignedResolution.Y;
 		const uint32 PixelGroupCount = PixelCount / PixelAlignment;
 		const uint32 FrameByteCount = PixelGroupCount * Info.PixelGroupSize;
@@ -121,16 +112,16 @@ namespace UE::RivermaxMediaCaptureUtil
 ///* URivermaxMediaCapture implementation
 //*****************************************************************************/
 
-UE::RivermaxCore::FRivermaxOutputStreamOptions URivermaxMediaCapture::GetOutputStreamOptions() const
+UE::RivermaxCore::FRivermaxOutputOptions URivermaxMediaCapture::GetOutputOptions() const
 {
 	return Options;
 }
 
 void URivermaxMediaCapture::GetLastPresentedFrameInformation(UE::RivermaxCore::FPresentedFrameInfo& OutFrameInfo) const
 {
-	if (RivermaxStream)
+	if (RivermaxVideoStream.IsValid())
 	{
-		RivermaxStream->GetLastPresentedFrame(OutFrameInfo);
+		RivermaxVideoStream->GetLastPresentedFrame(OutFrameInfo);
 	}
 }
 
@@ -149,7 +140,8 @@ bool URivermaxMediaCapture::ValidateMediaOutput() const
 bool URivermaxMediaCapture::InitializeCapture()
 {
 	URivermaxMediaOutput* RivermaxOutput = CastChecked<URivermaxMediaOutput>(MediaOutput);
-	bool bResult = Initialize(RivermaxOutput);
+	UE::RivermaxCore::FRivermaxOutputOptions MediaOutputOptions = RivermaxOutput->GenerateStreamOptions();
+	bool bResult = Initialize(MediaOutputOptions);
 #if WITH_EDITOR
 	if (bResult)
 	{
@@ -173,18 +165,24 @@ void URivermaxMediaCapture::StopCaptureImpl(bool bAllowPendingFrameToBeProcess)
 {
 	bIsActive = false;
 
-	if (RivermaxStream)
+	if (RivermaxVideoStream)
+	{
+		RivermaxVideoStream->Uninitialize();
+		RivermaxVideoStream.Reset();
+	}
+
+	if (RivermaxAncStream)
 	{
-		RivermaxStream->Uninitialize();
-		RivermaxStream.Reset();
+		RivermaxAncStream->Uninitialize();
+		RivermaxAncStream.Reset();
 	}
 }
 
 bool URivermaxMediaCapture::ShouldCaptureRHIResource() const
 {
-	if (RivermaxStream)
+	if (RivermaxVideoStream)
 	{
-		return RivermaxStream->IsGPUDirectSupported();
+		return RivermaxVideoStream->IsGPUDirectSupported();
 	}
 
 	return false;
@@ -195,22 +193,32 @@ bool URivermaxMediaCapture::HasFinishedProcessing() const
 	return Super::HasFinishedProcessing();
 }
 
-bool URivermaxMediaCapture::Initialize(URivermaxMediaOutput* InMediaOutput)
+bool URivermaxMediaCapture::Initialize(const UE::RivermaxCore::FRivermaxOutputOptions& InMediaOutputOptions)
 {
 	using namespace UE::RivermaxCore;
-
-	check(InMediaOutput);
 	bIsActive = false;
-
+	Options = InMediaOutputOptions;
 	IRivermaxCoreModule* Module = FModuleManager::GetModulePtr<IRivermaxCoreModule>("RivermaxCore");
 	if (Module)
 	{
-		if (ConfigureStream(InMediaOutput, Options))
+		TArray<char> SDP;
+		if (ConfigureCapture() && UE::RivermaxMediaUtils::Private::StreamOptionsToSDPDescription(Options, SDP))
 		{
-			RivermaxStream = Module->CreateOutputStream();
-			if (RivermaxStream)
+			using namespace UE::RivermaxCore;
+
+			RivermaxVideoStream = Module->CreateOutputStream(ERivermaxStreamType::VIDEO_2110_20_STREAM, SDP);
+			if (RivermaxVideoStream)
 			{
-				bIsActive = RivermaxStream->Initialize(Options, *this);
+				bIsActive = RivermaxVideoStream->Initialize(Options, *this);
+			}
+
+			if (Options.GetStreamOptions(ERivermaxStreamType::ANC_2110_40_STREAM).IsValid())
+			{
+				RivermaxAncStream = Module->CreateOutputStream(ERivermaxStreamType::ANC_2110_40_STREAM, SDP);
+				if (RivermaxAncStream)
+				{
+					bIsActive = RivermaxAncStream->Initialize(Options, *this);
+				}
 			}
 		}
 	}
@@ -218,7 +226,7 @@ bool URivermaxMediaCapture::Initialize(URivermaxMediaOutput* InMediaOutput)
 	return bIsActive;
 }
 
-bool URivermaxMediaCapture::ConfigureStream(URivermaxMediaOutput* InMediaOutput, UE::RivermaxCore::FRivermaxOutputStreamOptions& OutOptions) const
+bool URivermaxMediaCapture::ConfigureCapture() const
 {
 	using namespace UE::RivermaxCore;
 
@@ -228,43 +236,48 @@ bool URivermaxMediaCapture::ConfigureStream(URivermaxMediaOutput* InMediaOutput,
 	{
 		return false;
 	}
-	
+
 	if (!Module->GetRivermaxManager()->ValidateLibraryIsLoaded())
 	{
 		return false;
 	}
+	TSharedPtr<FRivermaxVideoOutputOptions> VideoOptions
+		= StaticCastSharedPtr<FRivermaxVideoOutputOptions>(Options.StreamOptions[static_cast<uint8>(ERivermaxStreamType::VIDEO_2110_20_STREAM)]);
 
-	const bool bFoundDevice = Module->GetRivermaxManager()->GetMatchingDevice(InMediaOutput->InterfaceAddress, OutOptions.InterfaceAddress);
-	if (bFoundDevice == false)
+	// Video configuration
+	if (VideoOptions.IsValid())
 	{
-		UE_LOG(LogRivermaxMedia, Error, TEXT("Could not find a matching interface for IP '%s'"), *InMediaOutput->InterfaceAddress);
-		return false;
-	}
+		const bool bFoundDevice = Module->GetRivermaxManager()->GetMatchingDevice(VideoOptions->InterfaceAddress, VideoOptions->InterfaceAddress);
+		if (bFoundDevice == false)
+		{
+			UE_LOG(LogRivermaxMedia, Error, TEXT("Could not find a matching interface for IP '%s'"), *VideoOptions->InterfaceAddress);
+			return false;
+		}
+		
+		/** Override the size. */
+		VideoOptions->Resolution = GetDesiredSize();
 
-	OutOptions.StreamAddress = InMediaOutput->StreamAddress;
-	OutOptions.Port = InMediaOutput->Port;
-	OutOptions.Resolution = GetDesiredSize();
+		if (VideoOptions->Resolution.X <= 0 || VideoOptions->Resolution.Y <= 0)
+		{
+			UE_LOG(LogRivermaxMedia, Warning, TEXT("Can't start capture. Invalid resolution requested: %dx%d"), VideoOptions->Resolution.X, VideoOptions->Resolution.Y);
+			return false;
+		}
 
-	if (OutOptions.Resolution.X <= 0 || OutOptions.Resolution.Y <= 0)
-	{
-		UE_LOG(LogRivermaxMedia, Warning, TEXT("Can't start capture. Invalid resolution requested: %dx%d"), OutOptions.Resolution.X, OutOptions.Resolution.Y);
-		return false;
+		const FVideoFormatInfo Info = FStandardVideoFormat::GetVideoFormatInfo(VideoOptions->PixelFormat);
+		VideoOptions->AlignedResolution = UE::RivermaxMediaUtils::Private::GetAlignedResolution(Info, VideoOptions->Resolution);
 	}
 
-	OutOptions.FrameRate = InMediaOutput->FrameRate;
-	OutOptions.NumberOfBuffers = InMediaOutput->PresentationQueueSize;
-	OutOptions.bUseGPUDirect = InMediaOutput->bUseGPUDirect;
-	OutOptions.AlignmentMode = UE::RivermaxMediaUtils::Private::MediaOutputAlignmentToRivermaxAlignment(InMediaOutput->AlignmentMode);
-	OutOptions.FrameLockingMode = UE::RivermaxMediaUtils::Private::MediaOutputFrameLockingToRivermax(InMediaOutput->FrameLockingMode);
-
-	// Setup alignment dependent configs
-	OutOptions.bDoContinuousOutput = OutOptions.AlignmentMode == ERivermaxAlignmentMode::AlignmentPoint ? InMediaOutput->bDoContinuousOutput : false;
-	OutOptions.bDoFrameCounterTimestamping = OutOptions.AlignmentMode == ERivermaxAlignmentMode::FrameCreation ? InMediaOutput->bDoFrameCounterTimestamping : false;
-
-	OutOptions.PixelFormat = UE::RivermaxMediaUtils::Private::MediaOutputPixelFormatToRivermaxSamplingType(InMediaOutput->PixelFormat);
-	const FVideoFormatInfo Info = FStandardVideoFormat::GetVideoFormatInfo(OutOptions.PixelFormat);
-	OutOptions.AlignedResolution = UE::RivermaxMediaCaptureUtil::GetAlignedResolution(Info, OutOptions.Resolution);
-
+	TSharedPtr<FRivermaxAncOutputOptions> AncOptions
+		= StaticCastSharedPtr<FRivermaxAncOutputOptions>(Options.StreamOptions[static_cast<uint8>(ERivermaxStreamType::ANC_2110_40_STREAM)]);
+	if (AncOptions.IsValid())
+	{
+		const bool bFoundDevice = Module->GetRivermaxManager()->GetMatchingDevice(AncOptions->InterfaceAddress, AncOptions->InterfaceAddress);
+		if (bFoundDevice == false)
+		{
+			UE_LOG(LogRivermaxMedia, Error, TEXT("Could not find a matching interface for IP '%s'"), *AncOptions->InterfaceAddress);
+			return false;
+		}
+	}
 	return true;
 }
 
@@ -286,7 +299,11 @@ void URivermaxMediaCapture::AddFrameReservationPass(FRDGBuilder& GraphBuilder)
 
 		if (Capturer->bIsActive)
 		{
-			Capturer->RivermaxStream->ReserveFrame();
+			Capturer->RivermaxVideoStream->ReserveFrame(FrameCounter);
+			if (Capturer->RivermaxAncStream.IsValid())
+			{
+				Capturer->RivermaxAncStream->ReserveFrame(FrameCounter);
+			}
 		}
 	});
 }
@@ -296,16 +313,20 @@ void URivermaxMediaCapture::OnFrameCapturedInternal_AnyThread(const FCaptureBase
 	using namespace UE::RivermaxCore;
 	TRACE_CPUPROFILER_EVENT_SCOPE_TEXT(*FRivermaxTracingUtils::RmaxOutMediaCapturePipeTraceEvents[InBaseData.SourceFrameNumberRenderThread % 10]);
 
-	FRivermaxOutputVideoFrameInfo NewFrame;
-	NewFrame.Height = Height;
-	NewFrame.Width = Width;
-	NewFrame.Stride = BytesPerRow;
-	NewFrame.CPUBuffer = InBuffer;
-	NewFrame.FrameIdentifier = InBaseData.SourceFrameNumberRenderThread;
-	if (RivermaxStream->PushVideoFrame(NewFrame) == false)
+	TSharedPtr<FRivermaxOutputVideoFrameInfo> NewFrame = MakeShared<FRivermaxOutputVideoFrameInfo>();
+	NewFrame->Height = Height;
+	NewFrame->Width = Width;
+	NewFrame->Stride = BytesPerRow;
+	NewFrame->CPUBuffer = InBuffer;
+	NewFrame->FrameIdentifier = InBaseData.SourceFrameNumberRenderThread;
+	if (RivermaxVideoStream->PushFrame(NewFrame) == false)
 	{
 		UE_LOG(LogRivermaxMedia, Verbose, TEXT("Failed to push captured frame"));
 	}
+	if (RivermaxAncStream.IsValid())
+	{
+		RivermaxAncStream->PushFrame(NewFrame);
+	}
 }
 
 void URivermaxMediaCapture::OnRHIResourceCapturedInternal_AnyThread(const FCaptureBaseData& InBaseData, TSharedPtr<FMediaCaptureUserData, ESPMode::ThreadSafe> InUserData, FBufferRHIRef InBuffer)
@@ -313,13 +334,17 @@ void URivermaxMediaCapture::OnRHIResourceCapturedInternal_AnyThread(const FCaptu
 	using namespace UE::RivermaxCore;
 	TRACE_CPUPROFILER_EVENT_SCOPE_TEXT(*FRivermaxTracingUtils::RmaxOutMediaCapturePipeTraceEvents[InBaseData.SourceFrameNumberRenderThread % 10]);
 
-	FRivermaxOutputVideoFrameInfo NewFrame;
-	NewFrame.FrameIdentifier = InBaseData.SourceFrameNumberRenderThread;
-	NewFrame.GPUBuffer = InBuffer;
-	if (RivermaxStream->PushVideoFrame(NewFrame) == false)
+	TSharedPtr<FRivermaxOutputVideoFrameInfo> NewFrame = MakeShared<FRivermaxOutputVideoFrameInfo>();
+	NewFrame->FrameIdentifier = InBaseData.SourceFrameNumberRenderThread;
+	NewFrame->GPUBuffer = InBuffer;
+	if (RivermaxVideoStream->PushFrame(NewFrame) == false)
 	{
 		UE_LOG(LogRivermaxMedia, Verbose, TEXT("Failed to pushed captured frame"));
 	}
+	if (RivermaxAncStream.IsValid())
+	{
+		RivermaxAncStream->PushFrame(NewFrame);
+	}
 }
 
 void URivermaxMediaCapture::OnFrameCaptured_RenderingThread(const FCaptureBaseData& InBaseData, TSharedPtr<FMediaCaptureUserData, ESPMode::ThreadSafe> InUserData, void* InBuffer, int32 Width, int32 Height, int32 BytesPerRow)
@@ -386,7 +411,7 @@ FRDGBufferDesc URivermaxMediaCapture::GetCustomBufferDescription(const FIntPoint
 
 	uint32 BytesPerElement = 0;
 	uint32 ElementsPerFrame = 0;
-	UE::RivermaxMediaCaptureUtil::GetOutputEncodingInfo(RivermaxOutput->PixelFormat, InDesiredSize, BytesPerElement, ElementsPerFrame);
+	UE::RivermaxMediaCaptureUtil::GetOutputEncodingInfo(RivermaxOutput->VideoStream.PixelFormat, InDesiredSize, BytesPerElement, ElementsPerFrame);
 	FRDGBufferDesc Desc = FRDGBufferDesc::CreateStructuredDesc(BytesPerElement, ElementsPerFrame);
 	
 	// Required when GPUDirect using CUDA will be involved
@@ -396,11 +421,12 @@ FRDGBufferDesc URivermaxMediaCapture::GetCustomBufferDescription(const FIntPoint
 
 void URivermaxMediaCapture::OnCustomCapture_RenderingThread(FRDGBuilder& GraphBuilder, const FCaptureBaseData& InBaseData, TSharedPtr<FMediaCaptureUserData, ESPMode::ThreadSafe> InUserData, FRDGTextureRef InSourceTexture, FRDGBufferRef OutputBuffer, const FRHICopyTextureInfo& CopyInfo, FVector2D CropU, FVector2D CropV)
 {
+	using namespace UE::RivermaxCore;
 	RDG_EVENT_SCOPE_STAT(GraphBuilder, Rmax_Capture, "Rmax_Capture");
 	RDG_GPU_STAT_SCOPE(GraphBuilder, Rmax_Capture);
 
 	TRACE_CPUPROFILER_EVENT_SCOPE(URivermaxMediaCapture::OnCustomCapture_RenderingThread);
-	TRACE_CPUPROFILER_EVENT_SCOPE_TEXT(*UE::RivermaxCore::FRivermaxTracingUtils::RmaxOutMediaCapturePipeTraceEvents[GFrameCounterRenderThread % 10]);
+	TRACE_CPUPROFILER_EVENT_SCOPE_TEXT(*FRivermaxTracingUtils::RmaxOutMediaCapturePipeTraceEvents[GFrameCounterRenderThread % 10]);
 
 	using namespace UE::RivermaxShaders;
 	URivermaxMediaOutput* RivermaxOutput = CastChecked<URivermaxMediaOutput>(MediaOutput);
@@ -410,15 +436,15 @@ void URivermaxMediaCapture::OnCustomCapture_RenderingThread(FRDGBuilder& GraphBu
 	// Rectangle area to use from source. This is used when source render target is bigger than output resolution
 	const FIntRect ViewRect(CopyInfo.GetSourceRect());
 	constexpr bool bDoLinearToSRGB = false;
-	const FIntPoint CaptureSize = Options.AlignedResolution;
+	const FIntPoint AlignedBufferSize = Options.GetStreamOptions<FRivermaxVideoOutputOptions>(ERivermaxStreamType::VIDEO_2110_20_STREAM)->AlignedResolution;
 	const FIntVector GroupCount = FComputeShaderUtils::GetGroupCount(DesiredOutputSize.X, 64);
 	
-	switch (RivermaxOutput->PixelFormat)
+	switch (RivermaxOutput->VideoStream.PixelFormat)
 	{
 	case ERivermaxMediaOutputPixelFormat::PF_8BIT_YUV422:
 	{
 		TShaderMapRef<FRGBToYUV8Bit422CS> ComputeShader(GlobalShaderMap);
-		FRGBToYUV8Bit422CS::FParameters* Parameters = ComputeShader->AllocateAndSetParameters(GraphBuilder, InSourceTexture, CaptureSize, ViewRect, DesiredOutputSize, MediaShaders::RgbToYuvRec709Scaled, MediaShaders::YUVOffset8bits, bDoLinearToSRGB, OutputBuffer);
+		FRGBToYUV8Bit422CS::FParameters* Parameters = ComputeShader->AllocateAndSetParameters(GraphBuilder, InSourceTexture, AlignedBufferSize, ViewRect, DesiredOutputSize, MediaShaders::RgbToYuvRec709Scaled, MediaShaders::YUVOffset8bits, bDoLinearToSRGB, OutputBuffer);
 
 		FComputeShaderUtils::AddPass(
 			GraphBuilder
@@ -432,7 +458,7 @@ void URivermaxMediaCapture::OnCustomCapture_RenderingThread(FRDGBuilder& GraphBu
 	case ERivermaxMediaOutputPixelFormat::PF_10BIT_YUV422:
 	{
 		TShaderMapRef<FRGBToYUV10Bit422LittleEndianCS> ComputeShader(GlobalShaderMap);
-		FRGBToYUV10Bit422LittleEndianCS::FParameters* Parameters = ComputeShader->AllocateAndSetParameters(GraphBuilder, InSourceTexture, CaptureSize, ViewRect, DesiredOutputSize, MediaShaders::RgbToYuvRec709Scaled, MediaShaders::YUVOffset10bits, bDoLinearToSRGB, OutputBuffer);
+		FRGBToYUV10Bit422LittleEndianCS::FParameters* Parameters = ComputeShader->AllocateAndSetParameters(GraphBuilder, InSourceTexture, AlignedBufferSize, ViewRect, DesiredOutputSize, MediaShaders::RgbToYuvRec709Scaled, MediaShaders::YUVOffset10bits, bDoLinearToSRGB, OutputBuffer);
 
 		FComputeShaderUtils::AddPass(
 			GraphBuilder
@@ -446,7 +472,7 @@ void URivermaxMediaCapture::OnCustomCapture_RenderingThread(FRDGBuilder& GraphBu
 	case ERivermaxMediaOutputPixelFormat::PF_8BIT_RGB:
 	{
 		TShaderMapRef<FRGBToRGB8BitCS> ComputeShader(GlobalShaderMap);
-		FRGBToRGB8BitCS::FParameters* Parameters = ComputeShader->AllocateAndSetParameters(GraphBuilder, InSourceTexture, CaptureSize, ViewRect, DesiredOutputSize, OutputBuffer);
+		FRGBToRGB8BitCS::FParameters* Parameters = ComputeShader->AllocateAndSetParameters(GraphBuilder, InSourceTexture, AlignedBufferSize, ViewRect, DesiredOutputSize, OutputBuffer);
 
 		FComputeShaderUtils::AddPass(
 			GraphBuilder
@@ -460,7 +486,7 @@ void URivermaxMediaCapture::OnCustomCapture_RenderingThread(FRDGBuilder& GraphBu
 	case ERivermaxMediaOutputPixelFormat::PF_10BIT_RGB:
 	{
 		TShaderMapRef<FRGBToRGB10BitCS> ComputeShader(GlobalShaderMap);
-		FRGBToRGB10BitCS::FParameters* Parameters = ComputeShader->AllocateAndSetParameters(GraphBuilder, InSourceTexture, CaptureSize, ViewRect, DesiredOutputSize, OutputBuffer);
+		FRGBToRGB10BitCS::FParameters* Parameters = ComputeShader->AllocateAndSetParameters(GraphBuilder, InSourceTexture, AlignedBufferSize, ViewRect, DesiredOutputSize, OutputBuffer);
 
 		FComputeShaderUtils::AddPass(
 			GraphBuilder
@@ -474,7 +500,7 @@ void URivermaxMediaCapture::OnCustomCapture_RenderingThread(FRDGBuilder& GraphBu
 	case ERivermaxMediaOutputPixelFormat::PF_12BIT_RGB:
 	{
 		TShaderMapRef<FRGBToRGB12BitCS> ComputeShader(GlobalShaderMap);
-		FRGBToRGB12BitCS::FParameters* Parameters = ComputeShader->AllocateAndSetParameters(GraphBuilder, InSourceTexture, CaptureSize, ViewRect, DesiredOutputSize, OutputBuffer);
+		FRGBToRGB12BitCS::FParameters* Parameters = ComputeShader->AllocateAndSetParameters(GraphBuilder, InSourceTexture, AlignedBufferSize, ViewRect, DesiredOutputSize, OutputBuffer);
 
 		FComputeShaderUtils::AddPass(
 			GraphBuilder
@@ -488,7 +514,7 @@ void URivermaxMediaCapture::OnCustomCapture_RenderingThread(FRDGBuilder& GraphBu
 	case ERivermaxMediaOutputPixelFormat::PF_FLOAT16_RGB:
 	{
 		TShaderMapRef<FRGBToRGB16fCS> ComputeShader(GlobalShaderMap);
-		FRGBToRGB16fCS::FParameters* Parameters = ComputeShader->AllocateAndSetParameters(GraphBuilder, InSourceTexture, CaptureSize, ViewRect, DesiredOutputSize, OutputBuffer);
+		FRGBToRGB16fCS::FParameters* Parameters = ComputeShader->AllocateAndSetParameters(GraphBuilder, InSourceTexture, AlignedBufferSize, ViewRect, DesiredOutputSize, OutputBuffer);
 
 		FComputeShaderUtils::AddPass(
 			GraphBuilder
diff --git a/Engine/Plugins/VirtualProduction/Rivermax/RivermaxMedia/Source/RivermaxMedia/Private/RivermaxMediaOutput.cpp b/Engine/Plugins/VirtualProduction/Rivermax/RivermaxMedia/Source/RivermaxMedia/Private/RivermaxMediaOutput.cpp
--- a/Engine/Plugins/VirtualProduction/Rivermax/RivermaxMedia/Source/RivermaxMedia/Private/RivermaxMediaOutput.cpp
+++ b/Engine/Plugins/VirtualProduction/Rivermax/RivermaxMedia/Source/RivermaxMedia/Private/RivermaxMediaOutput.cpp
@@ -2,9 +2,12 @@
 
 #include "RivermaxMediaOutput.h"
 
+#include "IRivermaxCoreModule.h"
 #include "MediaOutput.h"
+#include "Misc/FileHelper.h"
 #include "RivermaxMediaCapture.h"
-
+#include "RivermaxMediaLog.h"
+#include "RivermaxMediaUtils.h"
 
 /* URivermaxMediaOutput
 *****************************************************************************/
@@ -21,9 +24,9 @@ bool URivermaxMediaOutput::Validate(FString& OutFailureReason) const
 
 FIntPoint URivermaxMediaOutput::GetRequestedSize() const
 {
-	if (bOverrideResolution)
+	if (VideoStream.bOverrideResolution)
 	{
-		return Resolution;
+		return VideoStream.Resolution;
 	}
 
 	return UMediaOutput::RequestCaptureSourceSize;
@@ -39,7 +42,7 @@ EPixelFormat URivermaxMediaOutput::GetRequestedPixelFormat() const
 EMediaCaptureConversionOperation URivermaxMediaOutput::GetConversionOperation(EMediaCaptureSourceType InSourceType) const
 {
 	EMediaCaptureConversionOperation Result = EMediaCaptureConversionOperation::CUSTOM;
-	switch (PixelFormat)
+	switch (VideoStream.PixelFormat)
 	{
 	case ERivermaxMediaOutputPixelFormat::PF_8BIT_YUV422:
 	case ERivermaxMediaOutputPixelFormat::PF_10BIT_YUV422:
@@ -54,6 +57,100 @@ EMediaCaptureConversionOperation URivermaxMediaOutput::GetConversionOperation(EM
 	return Result;
 }
 
+void URivermaxMediaOutput::ExportSDP(const FString& InPath)
+{
+	TArray<char> SDP;
+	UE::RivermaxCore::FRivermaxOutputOptions OutputOptions = GenerateStreamOptions();
+
+	bool bSaved = false;
+	if (UE::RivermaxMediaUtils::Private::StreamOptionsToSDPDescription(OutputOptions, SDP))
+	{
+		FString SdpString(ANSI_TO_TCHAR(SDP.GetData()));
+
+		if (FFileHelper::SaveStringToFile(SdpString, *InPath))
+		{
+			bSaved = true;
+		}
+	}
+
+	if (bSaved)
+	{
+		UE_LOG(LogTemp, Log, TEXT("Saved SDP successfully to: %s"), *InPath);
+	}
+	else
+	{
+		UE_LOG(LogTemp, Error, TEXT("Failed to save SDP file: %s"), *InPath);
+	}
+}
+
+UE::RivermaxCore::FRivermaxOutputOptions URivermaxMediaOutput::GenerateStreamOptions() const
+{
+	using namespace UE::RivermaxCore;
+	using namespace UE::RivermaxMediaUtils::Private;
+
+	UE::RivermaxCore::FRivermaxOutputOptions OutOutputOptions;
+
+	// Video configuration
+	if (bOutputVideo)
+	{
+		// The following is preparation for UE5.7
+		{
+			VideoStream.bOverrideResolution = bOverrideResolution;
+			VideoStream.Resolution = Resolution;
+			VideoStream.FrameRate = FrameRate;
+			VideoStream.PixelFormat = PixelFormat;
+			VideoStream.InterfaceAddress = InterfaceAddress;
+			VideoStream.StreamAddress = StreamAddress;
+			VideoStream.Port = Port;
+			VideoStream.bUseGPUDirect = bUseGPUDirect;
+		}
+
+		TSharedPtr<FRivermaxVideoOutputOptions> VideoOptions
+			= StaticCastSharedPtr<FRivermaxVideoOutputOptions>(OutOutputOptions.StreamOptions[static_cast<uint8>(ERivermaxStreamType::VIDEO_2110_20_STREAM)]);
+
+		// Sanity check. Making sure that Video Options are not intialized prior.
+		check(!VideoOptions.IsValid());
+		VideoOptions = MakeShared<FRivermaxVideoOutputOptions>();
+		OutOutputOptions.StreamOptions[static_cast<uint8>(ERivermaxStreamType::VIDEO_2110_20_STREAM)] = VideoOptions;
+
+		VideoOptions->InterfaceAddress = VideoStream.InterfaceAddress;
+		VideoOptions->StreamAddress = VideoStream.StreamAddress;
+		VideoOptions->Port = VideoStream.Port;
+		VideoOptions->Resolution = GetRequestedSize();
+
+		VideoOptions->FrameRate = VideoStream.FrameRate;
+		OutOutputOptions.NumberOfBuffers = PresentationQueueSize;
+		VideoOptions->bUseGPUDirect = VideoStream.bUseGPUDirect;
+		OutOutputOptions.AlignmentMode = MediaOutputAlignmentToRivermaxAlignment(AlignmentMode);
+		OutOutputOptions.FrameLockingMode = UE::RivermaxMediaUtils::Private::MediaOutputFrameLockingToRivermax(FrameLockingMode);
+
+		// Setup alignment dependent configs
+		OutOutputOptions.bDoContinuousOutput = OutOutputOptions.AlignmentMode == ERivermaxAlignmentMode::AlignmentPoint ? bDoContinuousOutput : false;
+		OutOutputOptions.bDoFrameCounterTimestamping = OutOutputOptions.AlignmentMode == ERivermaxAlignmentMode::FrameCreation ? bDoFrameCounterTimestamping : false;
+
+		VideoOptions->PixelFormat = MediaOutputPixelFormatToRivermaxSamplingType(VideoStream.PixelFormat);
+		const FVideoFormatInfo Info = FStandardVideoFormat::GetVideoFormatInfo(VideoOptions->PixelFormat);
+		VideoOptions->AlignedResolution = GetAlignedResolution(Info, VideoOptions->Resolution);
+	}
+
+	if (bOutputAnc)
+	{
+		TSharedPtr<FRivermaxAncOutputOptions> AncOptions
+			= StaticCastSharedPtr<FRivermaxAncOutputOptions>(OutOutputOptions.StreamOptions[static_cast<uint8>(ERivermaxStreamType::ANC_2110_40_STREAM)]);
+		
+		// Sanity check. Making sure that Anc Options are not intialized prior.
+		check(!AncOptions.IsValid());
+		AncOptions = MakeShared<FRivermaxAncOutputOptions>();
+		OutOutputOptions.StreamOptions[static_cast<uint8>(ERivermaxStreamType::ANC_2110_40_STREAM)] = AncOptions;
+
+		AncOptions->InterfaceAddress = AncStream.InterfaceAddress;
+		AncOptions->StreamAddress = AncStream.StreamAddress;
+		AncOptions->Port = AncStream.Port;
+	}
+
+	return MoveTemp(OutOutputOptions);
+}
+
 UMediaCapture* URivermaxMediaOutput::CreateMediaCaptureImpl()
 {
 	UMediaCapture* Result = NewObject<URivermaxMediaCapture>();
diff --git a/Engine/Plugins/VirtualProduction/Rivermax/RivermaxMedia/Source/RivermaxMedia/Private/RivermaxMediaPlayer.cpp b/Engine/Plugins/VirtualProduction/Rivermax/RivermaxMedia/Source/RivermaxMedia/Private/RivermaxMediaPlayer.cpp
--- a/Engine/Plugins/VirtualProduction/Rivermax/RivermaxMedia/Source/RivermaxMedia/Private/RivermaxMediaPlayer.cpp
+++ b/Engine/Plugins/VirtualProduction/Rivermax/RivermaxMedia/Source/RivermaxMedia/Private/RivermaxMediaPlayer.cpp
@@ -146,7 +146,11 @@ namespace UE::RivermaxMedia
 		IRivermaxCoreModule* Module = FModuleManager::GetModulePtr<IRivermaxCoreModule>("RivermaxCore");
 		if (Module && ConfigureStream(Options))
 		{
-			InputStream = Module->CreateInputStream();
+			
+			TArray<char> SDP;
+			// TODO: Similar to Output stream creation, input streams should be created from the supplied SDP.
+			//UE::RivermaxMediaUtils::Private::StreamOptionsToSDPDescription(Options, SDP);
+			InputStream = Module->CreateInputStream(UE::RivermaxCore::ERivermaxStreamType::VIDEO_2110_20_STREAM, SDP);
 		}
 
 		// If we are not following the stream resolution, make it the video track format and then reset to go through a format change once
diff --git a/Engine/Plugins/VirtualProduction/Rivermax/RivermaxMedia/Source/RivermaxMedia/Private/RivermaxMediaUtils.cpp b/Engine/Plugins/VirtualProduction/Rivermax/RivermaxMedia/Source/RivermaxMedia/Private/RivermaxMediaUtils.cpp
--- a/Engine/Plugins/VirtualProduction/Rivermax/RivermaxMedia/Source/RivermaxMedia/Private/RivermaxMediaUtils.cpp
+++ b/Engine/Plugins/VirtualProduction/Rivermax/RivermaxMedia/Source/RivermaxMedia/Private/RivermaxMediaUtils.cpp
@@ -2,6 +2,7 @@
 
 #include "RivermaxMediaUtils.h"
 
+#include "RivermaxMediaLog.h"
 #include "RivermaxShaders.h"
 #include "Serialization/CustomVersion.h"
 
@@ -141,7 +142,7 @@ namespace UE::RivermaxMediaUtils::Private
 		}
 		default:
 		{
-			checkNoEntry(); 
+			checkNoEntry();
 			return Description;
 		}
 		}
@@ -149,10 +150,20 @@ namespace UE::RivermaxMediaUtils::Private
 		// Shader encoding might not align with pixel group size so we need to have enough elements to represent the last pixel group
 		Description.NumberOfElements = FrameByteCount / Description.BytesPerElement;
 		Description.NumberOfElements += FrameByteCount % Description.BytesPerElement != 0 ? 1 : 0;
-		
+
 		return Description;
 	}
 
+	FIntPoint GetAlignedResolution(const UE::RivermaxCore::FVideoFormatInfo& InFormatInfo, const FIntPoint& ResolutionToAlign)
+	{
+		// We need to adjust horizontal resolution if it's not aligned to pgroup pixel coverage
+		// RTP header has to describe a full pgroup even if it's outside of effective resolution. Refer to st2110-20: 6.2.2 pgroup Size and Construction  
+		const uint32 PixelAlignment = InFormatInfo.PixelGroupCoverage;
+		uint32 PixelsToBeFitIntoGroup = ResolutionToAlign.X % PixelAlignment;
+		const uint32 AlignedHorizontalResolution = PixelsToBeFitIntoGroup > 0 ? ResolutionToAlign.X + (PixelAlignment - PixelsToBeFitIntoGroup) : ResolutionToAlign.X;
+		return FIntPoint(AlignedHorizontalResolution, ResolutionToAlign.Y);
+	}
+
 	UE::RivermaxCore::ERivermaxAlignmentMode MediaOutputAlignmentToRivermaxAlignment(ERivermaxMediaAlignmentMode InAlignmentMode)
 	{
 		using namespace UE::RivermaxCore;
@@ -251,6 +262,226 @@ namespace UE::RivermaxMediaUtils::Private
 		}
 	}
 
+	FString PixelFormatToSamplingDesc(UE::RivermaxCore::ESamplingType SamplingType)
+	{
+		using namespace UE::RivermaxCore;
+		switch (SamplingType)
+		{
+		case ESamplingType::RGB_8bit:
+		case ESamplingType::RGB_10bit:
+		case ESamplingType::RGB_12bit:
+		case ESamplingType::RGB_16bit:
+		case ESamplingType::RGB_16bitFloat:
+		{
+			return FString(TEXT("RGB"));
+		}
+		case ESamplingType::YUV422_8bit:
+		case ESamplingType::YUV422_10bit:
+		case ESamplingType::YUV422_12bit:
+		case ESamplingType::YUV422_16bit:
+		case ESamplingType::YUV422_16bitFloat:
+		{
+			return FString(TEXT("YCbCr-4:2:2"));
+		}
+		case ESamplingType::YUV444_8bit:
+		case ESamplingType::YUV444_10bit:
+		case ESamplingType::YUV444_12bit:
+		case ESamplingType::YUV444_16bit:
+		case ESamplingType::YUV444_16bitFloat:
+		{
+			return FString(TEXT("YCbCr-4:4:4"));
+		}
+		default:
+		{
+			checkNoEntry();
+			return FString();
+		}
+		}
+	}
+
+	FString PixelFormatToBitDepth(UE::RivermaxCore::ESamplingType SamplingType)
+	{
+		using namespace UE::RivermaxCore;
+		switch (SamplingType)
+		{
+		case ESamplingType::RGB_8bit:
+		case ESamplingType::YUV422_8bit:
+		case ESamplingType::YUV444_8bit:
+		{
+			return FString(TEXT("8"));
+		}
+		case ESamplingType::RGB_10bit:
+		case ESamplingType::YUV422_10bit:
+		case ESamplingType::YUV444_10bit:
+		{
+			return FString(TEXT("10"));
+		}
+		case ESamplingType::RGB_12bit:
+		case ESamplingType::YUV422_12bit:
+		case ESamplingType::YUV444_12bit:
+		{
+			return FString(TEXT("12"));
+		}
+		case ESamplingType::RGB_16bit:
+		case ESamplingType::YUV422_16bit:
+		case ESamplingType::YUV444_16bit:
+		{
+			return FString(TEXT("16"));
+		}
+		case ESamplingType::RGB_16bitFloat:
+		case ESamplingType::YUV422_16bitFloat:
+		case ESamplingType::YUV444_16bitFloat:
+		{
+			return FString(TEXT("16f"));
+		}
+		default:
+		{
+			checkNoEntry();
+			return FString();
+		}
+		}
+	}
+
+	bool StreamOptionsToSDPDescription(const UE::RivermaxCore::FRivermaxOutputOptions& Options, TArray<char>& OutSDPDescription)
+	{
+		uint8 NumberOfStreams = 0; 
+
+		// There needs to be at least one valid stream for SDP to be created.
+		for (const TSharedPtr<UE::RivermaxCore::FRivermaxOutputStreamOptions>& Stream : Options.StreamOptions)
+		{
+			if (Stream.IsValid())
+			{
+				NumberOfStreams++;
+			}
+		}
+
+		if (NumberOfStreams == 0)
+		{
+			return false;
+		}
+
+		// Basic SDP string creation from a set of options. At some point, having a proper SDP loader / creator would be useful.
+		// Refer to https://datatracker.ietf.org/doc/html/rfc4570
+
+
+		constexpr int32 MulticastTTL = 32;
+		TAnsiStringBuilder<2048> SDPDescriptionBuilder;
+
+		// Protocol version.
+		SDPDescriptionBuilder.Appendf("v=0\n");
+
+		// Session name.
+		SDPDescriptionBuilder.Appendf("s=Unreal Engine Rivermax 2110 \n");
+
+		// Start and stop time of the session. 0 0 means permanently active.
+		SDPDescriptionBuilder.Appendf("t=0 0\n");
+
+		// Group streams. Currently unsupported by Rivermax.
+		if (NumberOfStreams > 1)
+		{
+			// This isn't supported by Rivermax 1.60 and the Rivermax sdp parser will fail without giving a clear reason.
+			// Keeping it here for potential future support.
+			//SDPDescriptionBuilder.Appendf("a=group:FID V1 M1\n");
+		}
+
+		// This is used by RivermaxOutStream when it calls to the library to assign Media Block Index in SDP.
+		uint64 StreamIndex = 0;
+
+		using namespace UE::RivermaxCore;
+
+		// Adding Video stream properties to the SDP
+		TSharedPtr<FRivermaxVideoOutputOptions> VideoOptions 
+			= StaticCastSharedPtr<FRivermaxVideoOutputOptions>(Options.StreamOptions[static_cast<uint8>(ERivermaxStreamType::VIDEO_2110_20_STREAM)]);
+
+		if (VideoOptions.IsValid())
+		{
+
+			// Apply desired multiplier and convert fractional frame rate to be represented over 1001 for sdp compliance.
+			FFrameRate DesiredRate = VideoOptions->FrameRate;
+			const double DecimalDesiredRate = VideoOptions->FrameRate.AsDecimal();
+			const bool bIsNonStandardFractionalFrameRate = (DesiredRate.Denominator != 1001 && !FMath::IsNearlyZero(FMath::Frac(DecimalDesiredRate)));
+			if (bIsNonStandardFractionalFrameRate)
+			{
+				UE_LOG(LogRivermaxMedia, Warning, TEXT("Fractional frame rates must be described using a denominator of 1001. Converting it for stream creation."));
+			}
+
+			FString FrameRateDescription;
+			if (bIsNonStandardFractionalFrameRate)
+			{
+				const double NewRateDecimal = DecimalDesiredRate * 1001;
+				DesiredRate.Numerator = FMath::RoundToInt32(NewRateDecimal);
+				DesiredRate.Denominator = 1001;
+			}
+
+			if (DesiredRate.Denominator == 1001)
+			{
+				FrameRateDescription = FString::Printf(TEXT("%u/%u"), DesiredRate.Numerator, DesiredRate.Denominator);
+			}
+			else
+			{
+				FrameRateDescription = FString::Printf(TEXT("%u"), FMath::RoundToInt32(DesiredRate.AsDecimal()));
+			}
+
+			// Media description.
+			SDPDescriptionBuilder.Appendf("m=video %d RTP/AVP 96\n", VideoOptions->Port);
+
+			// Connection information.
+			SDPDescriptionBuilder.Appendf("c=IN IP4 %S/%d\n", *VideoOptions->StreamAddress, MulticastTTL);
+
+			// Limits multicast to a specific device.
+			SDPDescriptionBuilder.Appendf("a=source-filter: incl IN IP4 %S %S\n", *VideoOptions->StreamAddress, *VideoOptions->InterfaceAddress);
+
+			// RTP mapping attribute. Payload type and clock rate.
+			SDPDescriptionBuilder.Appendf("a=rtpmap:96 raw/90000\n");
+
+			// Video Format parameters.
+			SDPDescriptionBuilder.Appendf("a=fmtp: 96 sampling=%S; width=%d; height=%d; exactframerate=%S; depth=%S; TCS=SDR; colorimetry=BT709; PM=2110GPM; SSN=ST2110-20:2017; TP=2110TPN;\n"
+				, *PixelFormatToSamplingDesc(VideoOptions->PixelFormat)
+				, VideoOptions->AlignedResolution.X
+				, VideoOptions->AlignedResolution.Y
+				, *FrameRateDescription
+				, *PixelFormatToBitDepth(VideoOptions->PixelFormat));
+
+			// sets the identifier for this stream.
+			SDPDescriptionBuilder.Appendf("a=mid:V1\n");
+
+			VideoOptions->StreamIndex = StreamIndex++;
+		}
+
+
+		TSharedPtr<FRivermaxAncOutputOptions> AncOptions
+			= StaticCastSharedPtr<FRivermaxAncOutputOptions>(Options.StreamOptions[static_cast<uint8>(ERivermaxStreamType::ANC_2110_40_STREAM)]);
+
+		// Adding Anc stream properties to the SDP
+		if (AncOptions.IsValid())
+		{
+			// Media description. 
+			SDPDescriptionBuilder.Appendf("m=video %d RTP/AVP 97\n", AncOptions->Port);
+
+			// Connection information.
+			SDPDescriptionBuilder.Appendf("c=IN IP4 %S/%d\n", *AncOptions->StreamAddress, MulticastTTL);
+
+			// Limits multicast to a specific device.
+			SDPDescriptionBuilder.Appendf("a=source-filter: incl IN IP4 %S %S\n", *AncOptions->StreamAddress, *AncOptions->InterfaceAddress);
+
+			// RTP mapping attribute. Payload type and clock rate.
+			SDPDescriptionBuilder.Appendf("a=rtpmap:97 smpte291/90000\n");
+
+			// RTP mapping attribute. Payload type and clock rate.
+			SDPDescriptionBuilder.Appendf("a=fmtp:97 DID_SDID={0x%X,0x%X}\n", AncOptions->DID, AncOptions->SDID);
+
+			// sets the identifier for this stream.
+			SDPDescriptionBuilder.Appendf("a=mid:M1\n");
+
+			AncOptions->StreamIndex = StreamIndex++;
+		}
+
+		OutSDPDescription.SetNum(SDPDescriptionBuilder.Len());
+		memcpy(&OutSDPDescription[0], SDPDescriptionBuilder.GetData(), SDPDescriptionBuilder.Len());
+		OutSDPDescription.Add('\0');
+		return true;
+	}
+
 	UE::RivermaxCore::EFrameLockingMode MediaOutputFrameLockingToRivermax(ERivermaxFrameLockingMode InFrameLockingMode)
 	{
 		using namespace UE::RivermaxCore;
diff --git a/Engine/Plugins/VirtualProduction/Rivermax/RivermaxMedia/Source/RivermaxMedia/Private/RivermaxMediaUtils.h b/Engine/Plugins/VirtualProduction/Rivermax/RivermaxMedia/Source/RivermaxMedia/Private/RivermaxMediaUtils.h
--- a/Engine/Plugins/VirtualProduction/Rivermax/RivermaxMedia/Source/RivermaxMedia/Private/RivermaxMediaUtils.h
+++ b/Engine/Plugins/VirtualProduction/Rivermax/RivermaxMedia/Source/RivermaxMedia/Private/RivermaxMediaUtils.h
@@ -27,6 +27,14 @@ namespace UE::RivermaxMediaUtils::Private
 	UE::RivermaxCore::EFrameLockingMode MediaOutputFrameLockingToRivermax(ERivermaxFrameLockingMode InFrameLockingMode);
 	FSourceBufferDesc GetBufferDescription(const FIntPoint& Resolution, ERivermaxMediaSourcePixelFormat InPixelFormat);
 
+	/** Gets a resolution aligned to a pixel group specified for the provided video format as per 2110-20 requirements. */
+	FIntPoint GetAlignedResolution(const UE::RivermaxCore::FVideoFormatInfo& InFormatInfo, const FIntPoint& ResolutionToAlign);
+
+	/** 
+	* Convert a set of streaming option to its SDP description. Currently only support video type. 
+	* Returns true if SDP could successfully be created from rivermax options.
+	*/
+	bool StreamOptionsToSDPDescription(const UE::RivermaxCore::FRivermaxOutputOptions& Options, TArray<char>& OutSDPDescription);
 }
 
 // Custom version to keep track of and restore depricated properties.
diff --git a/Engine/Plugins/VirtualProduction/Rivermax/RivermaxMedia/Source/RivermaxMedia/Public/RivermaxMediaCapture.h b/Engine/Plugins/VirtualProduction/Rivermax/RivermaxMedia/Source/RivermaxMedia/Public/RivermaxMediaCapture.h
--- a/Engine/Plugins/VirtualProduction/Rivermax/RivermaxMedia/Source/RivermaxMedia/Public/RivermaxMediaCapture.h
+++ b/Engine/Plugins/VirtualProduction/Rivermax/RivermaxMedia/Source/RivermaxMedia/Public/RivermaxMediaCapture.h
@@ -24,7 +24,7 @@ class RIVERMAXMEDIA_API URivermaxMediaCapture : public UMediaCapture, public UE:
 
 public:
 	/** Rivermax capture specific API to provide stream options access */
-	UE::RivermaxCore::FRivermaxOutputStreamOptions GetOutputStreamOptions() const;
+	UE::RivermaxCore::FRivermaxOutputOptions GetOutputOptions() const;
 
 	/** Returns information about last presented frame on the output stream */
 	void GetLastPresentedFrameInformation(UE::RivermaxCore::FPresentedFrameInfo& OutFrameInfo) const;
@@ -74,11 +74,15 @@ protected:
 private:
 	struct FRivermaxCaptureSyncData;
 
-	/** Initializes capture and launches stream creation */
-	bool Initialize(URivermaxMediaOutput* InMediaOutput);
+	/** Initializes capture and launches stream creation. */
+	bool Initialize(const UE::RivermaxCore::FRivermaxOutputOptions& InMediaOutput);
 
-	/** Configures rivermax stream with desired output options */
-	bool ConfigureStream(URivermaxMediaOutput* InMediaOutput, UE::RivermaxCore::FRivermaxOutputStreamOptions& OutOptions) const;
+	/**
+	 * Checks whether the provided options are valid to start the capture.
+	 * Resolves the interface address and ensures the resolution is aligned to the nearest pixel group.
+	 * Returns false if the stream cannot be configured based on the provided settings.
+	 */
+	bool ConfigureCapture() const;
 
 	/** Enqueues a RHI lambda to reserve a spot for the next frame to capture */
 	void AddFrameReservationPass(FRDGBuilder& GraphBuilder);
@@ -88,13 +92,17 @@ private:
 
 	/** Common method called for gpudirect route when a frame is captured, either from render thread or any thread  */
 	void OnRHIResourceCapturedInternal_AnyThread(const FCaptureBaseData& InBaseData, TSharedPtr<FMediaCaptureUserData, ESPMode::ThreadSafe> InUserData, FBufferRHIRef InBuffer);
+
 private:
 
-	/** Instance of the rivermax stream opened for this capture */
-	TUniquePtr<UE::RivermaxCore::IRivermaxOutputStream> RivermaxStream;
+	/** Instance of the rivermax video stream opened for this capture */
+	TUniquePtr<UE::RivermaxCore::IRivermaxOutputStream> RivermaxVideoStream;
+
+	/** Instance of the rivermax ancillary stream opened for this capture */
+	TUniquePtr<UE::RivermaxCore::IRivermaxOutputStream> RivermaxAncStream;
 
 	/** Set of options used to configure output stream */
-	UE::RivermaxCore::FRivermaxOutputStreamOptions Options;
+	UE::RivermaxCore::FRivermaxOutputOptions Options;
 
 	/** Whether capture is active. Used to be queried by sync task */
 	std::atomic<bool> bIsActive;
diff --git a/Engine/Plugins/VirtualProduction/Rivermax/RivermaxMedia/Source/RivermaxMedia/Public/RivermaxMediaOutput.h b/Engine/Plugins/VirtualProduction/Rivermax/RivermaxMedia/Source/RivermaxMedia/Public/RivermaxMediaOutput.h
--- a/Engine/Plugins/VirtualProduction/Rivermax/RivermaxMedia/Source/RivermaxMedia/Public/RivermaxMediaOutput.h
+++ b/Engine/Plugins/VirtualProduction/Rivermax/RivermaxMedia/Source/RivermaxMedia/Public/RivermaxMediaOutput.h
@@ -50,6 +50,64 @@ enum class ERivermaxFrameLockingMode : uint8
 	BlockOnReservation,
 };
 
+/** Basic settings for all 2110 stream types . */
+USTRUCT(BlueprintType)
+struct RIVERMAXMEDIA_API FRivermaxStream
+{
+	GENERATED_BODY()
+
+	/**
+	 * Network card interface to use to send data
+	 * Wildcards are supported to match against an interface found on the machine
+	 * 192.*.0.110
+	 * 192.168.0.1?0
+	 * 192.168.0.1*
+	 */
+	UPROPERTY(EditAnywhere, BlueprintReadWrite, Category = "Settings")
+	FString InterfaceAddress = TEXT("*.*.*.*");
+
+	/** Address of the stream. Can be multicast, i.e. 228.1.1.1) */
+	UPROPERTY(EditAnywhere, BlueprintReadWrite, Category = "Settings")
+	FString StreamAddress = UE::RivermaxCore::DefaultStreamAddress;
+
+	/** Port to use for this output */
+	UPROPERTY(EditAnywhere, BlueprintReadWrite, Category = "Settings")
+	int32 Port = 50000;
+};
+
+/** All settings required to initialize Video (2110-20) stream. */
+USTRUCT(BlueprintType)
+struct RIVERMAXMEDIA_API FRivermaxVideoStream : public FRivermaxStream
+{
+	GENERATED_BODY()
+
+	/** Frame rate of this output stream */
+	UPROPERTY(EditAnywhere, BlueprintReadWrite, Category = "Settings")
+	FFrameRate FrameRate = { 24,1 };
+
+	/** If false, use the size provided by the source. If true, a specific resolution will be used. */
+	UPROPERTY(BlueprintReadWrite, EditAnywhere, Category = "Media", meta = (InlineEditConditionToggle))
+	bool bOverrideResolution = false;
+
+	/** Resolution of this output stream */
+	UPROPERTY(EditAnywhere, BlueprintReadWrite, Category = "Settings", meta = (EditCondition = "bOverrideResolution"))
+	FIntPoint Resolution = { 1920, 1080 };
+
+	/** Pixel format for this output stream */
+	UPROPERTY(EditAnywhere, BlueprintReadWrite, Category = "Settings")
+	ERivermaxMediaOutputPixelFormat PixelFormat = ERivermaxMediaOutputPixelFormat::PF_10BIT_RGB;
+
+	/** Whether to use GPUDirect if available (Memcopy from GPU to NIC directly bypassing system) if available */
+	UPROPERTY(EditAnywhere, AdvancedDisplay, Category = "Video")
+	bool bUseGPUDirect = true;
+};
+
+/** All settings required to initialize ANC (2110-40) stream. */
+USTRUCT(BlueprintType)
+struct RIVERMAXMEDIA_API FRivermaxAncStream : public FRivermaxStream
+{
+	GENERATED_BODY()
+};
 
 /**
  * Output information for a Rivermax media capture.
@@ -67,6 +125,14 @@ public:
 	virtual EPixelFormat GetRequestedPixelFormat() const override;
 	virtual EMediaCaptureConversionOperation GetConversionOperation(EMediaCaptureSourceType InSourceType) const override;
 
+public:
+
+	/** Saves SDP into the file fath. */
+	void ExportSDP(const FString& InPath);
+
+	/** Aggregates URivermaxMediaOutput settings into FRivermaxOutputOptions that have information about individual stream types. */
+	UE::RivermaxCore::FRivermaxOutputOptions GenerateStreamOptions() const;
+
 protected:
 	virtual UMediaCapture* CreateMediaCaptureImpl() override;
 	//~ End UMediaOutput interface
@@ -106,18 +172,20 @@ public:
 	UPROPERTY(EditAnywhere, Category = "Output", meta = (EditCondition = "AlignmentMode == ERivermaxMediaAlignmentMode::FrameCreation"))
 	bool bDoFrameCounterTimestamping = true;
 
+
+public:
 	/** If false, use the default source buffer size. If true, a specific resolution will be used. */
 	UPROPERTY(BlueprintReadWrite, EditAnywhere, Category = "Media", meta = (InlineEditConditionToggle))
 	bool bOverrideResolution = false;
 
 	/** Resolution of this output stream */
-	UPROPERTY(EditAnywhere, BlueprintReadWrite, Category="Settings", meta = (EditCondition = "bOverrideResolution"))
-	FIntPoint Resolution = {1920, 1080};
-	
+	UPROPERTY(EditAnywhere, BlueprintReadWrite, Category = "Settings", meta = (EditCondition = "bOverrideResolution"))
+	FIntPoint Resolution = { 1920, 1080 };
+
 	/** Frame rate of this output stream */
 	UPROPERTY(EditAnywhere, BlueprintReadWrite, Category = "Settings")
-	FFrameRate FrameRate = {24,1};
-	
+	FFrameRate FrameRate = { 24,1 };
+
 	/** Pixel format for this output stream */
 	UPROPERTY(EditAnywhere, BlueprintReadWrite, Category = "Settings")
 	ERivermaxMediaOutputPixelFormat PixelFormat = ERivermaxMediaOutputPixelFormat::PF_10BIT_RGB;
@@ -128,7 +196,7 @@ public:
 	 * 192.*.0.110
 	 * 192.168.0.1?0
 	 * 192.168.0.1*
-	 */	
+	 */
 	UPROPERTY(EditAnywhere, BlueprintReadWrite, Category = "Settings")
 	FString InterfaceAddress = TEXT("*.*.*.*");
 
@@ -143,4 +211,13 @@ public:
 	/** Whether to use GPUDirect if available (Memcopy from GPU to NIC directly bypassing system) if available */
 	UPROPERTY(EditAnywhere, AdvancedDisplay, Category = "Video")
 	bool bUseGPUDirect = true;
+
+private:
+	friend class URivermaxMediaCapture;
+
+	bool bOutputVideo = true;
+	mutable FRivermaxVideoStream VideoStream;
+
+	bool bOutputAnc = false;
+	mutable FRivermaxAncStream AncStream;
 };
diff --git a/Engine/Plugins/VirtualProduction/Rivermax/RivermaxMedia/Source/RivermaxMediaEditor/Private/Customization/RivermaxMediaDetailsCustomization.cpp b/Engine/Plugins/VirtualProduction/Rivermax/RivermaxMedia/Source/RivermaxMediaEditor/Private/Customization/RivermaxMediaDetailsCustomization.cpp
--- a/Engine/Plugins/VirtualProduction/Rivermax/RivermaxMedia/Source/RivermaxMediaEditor/Private/Customization/RivermaxMediaDetailsCustomization.cpp
+++ b/Engine/Plugins/VirtualProduction/Rivermax/RivermaxMedia/Source/RivermaxMediaEditor/Private/Customization/RivermaxMediaDetailsCustomization.cpp
@@ -5,8 +5,13 @@
 #include "Customizations/RivermaxDeviceSelectionCustomization.h"
 #include "DetailLayoutBuilder.h"
 #include "DetailWidgetRow.h"
+#include "DesktopPlatformModule.h"
+#include "Framework/Application/SlateApplication.h"
+#include "IDesktopPlatform.h"
 #include "RivermaxMediaOutput.h"
 #include "RivermaxMediaSource.h"
+#include "Widgets/Input/SButton.h"
+
 
 TSharedRef<IDetailCustomization> FRivermaxMediaDetailsCustomization::MakeInstance()
 {
@@ -22,13 +27,52 @@ void FRivermaxMediaDetailsCustomization::CustomizeDetails(IDetailLayoutBuilder&
 	{
 		// For now, take care of both Rivermax source and output to customize their interface field
 		const TWeakObjectPtr<UObject> Obj = CustomizedObjects[ObjectIndex];
+
 		if (URivermaxMediaOutput* Output = Cast<URivermaxMediaOutput>(Obj.Get()))
 		{
-			UE::RivermaxCore::Utils::SetupDeviceSelectionCustomization(ObjectIndex, Output->InterfaceAddress, DetailBuilder.GetProperty(GET_MEMBER_NAME_CHECKED(URivermaxMediaOutput, InterfaceAddress)), DetailBuilder);
+			UE::RivermaxCore::Utils::SetupDeviceSelectionCustomization(ObjectIndex, Output->InterfaceAddress, DetailBuilder.GetProperty(GET_MEMBER_NAME_CHECKED(FRivermaxStream, InterfaceAddress)), DetailBuilder);
 		}
-		else if (URivermaxMediaSource* Source = Cast<URivermaxMediaSource>(Obj.Get()))
+		else  if (URivermaxMediaSource* Source = Cast<URivermaxMediaSource>(Obj.Get()))
 		{
 			UE::RivermaxCore::Utils::SetupDeviceSelectionCustomization(ObjectIndex, Source->InterfaceAddress, DetailBuilder.GetProperty(GET_MEMBER_NAME_CHECKED(URivermaxMediaSource, InterfaceAddress)), DetailBuilder);
 		}
+
+		if (URivermaxMediaOutput* Output = Cast<URivermaxMediaOutput>(Obj.Get()))
+		{
+			IDetailCategoryBuilder& Category = DetailBuilder.EditCategory("Advanced", FText::FromString("Advanced"), ECategoryPriority::Important);
+			Category.AddCustomRow(FText::FromString("Export SDP"))
+				.ValueContent()
+				[
+					SNew(SButton)
+						.Text(FText::FromString("Export SDP"))
+						.OnClicked(FOnClicked::CreateRaw(this, &FRivermaxMediaDetailsCustomization::OnExportSDP, Output))
+				];
+		}
+	}
+}
+
+FReply FRivermaxMediaDetailsCustomization::OnExportSDP(URivermaxMediaOutput* Output)
+{
+	IDesktopPlatform* DesktopPlatform = FDesktopPlatformModule::Get();
+	if (DesktopPlatform)
+	{
+		TArray<FString> OutFilenames;
+		const void* ParentWindowHandle = FSlateApplication::Get().FindBestParentWindowHandleForDialogs(nullptr);
+
+		bool bSuccess = DesktopPlatform->SaveFileDialog(
+			ParentWindowHandle,
+			TEXT("Export SDP File"),
+			FPaths::ProjectDir(),
+			TEXT("Rivermax Media Output.sdp"),
+			TEXT("SDP Files (*.sdp)|*.sdp"),
+			EFileDialogFlags::None,
+			OutFilenames
+		);
+
+		if (bSuccess && OutFilenames.Num() > 0)
+		{
+			Output->ExportSDP(OutFilenames[0]);
+		}
 	}
+	return FReply::Handled();
 }
diff --git a/Engine/Plugins/VirtualProduction/Rivermax/RivermaxMedia/Source/RivermaxMediaEditor/Private/Customization/RivermaxMediaDetailsCustomization.h b/Engine/Plugins/VirtualProduction/Rivermax/RivermaxMedia/Source/RivermaxMediaEditor/Private/Customization/RivermaxMediaDetailsCustomization.h
--- a/Engine/Plugins/VirtualProduction/Rivermax/RivermaxMedia/Source/RivermaxMediaEditor/Private/Customization/RivermaxMediaDetailsCustomization.h
+++ b/Engine/Plugins/VirtualProduction/Rivermax/RivermaxMedia/Source/RivermaxMediaEditor/Private/Customization/RivermaxMediaDetailsCustomization.h
@@ -3,9 +3,10 @@
 #pragma once
 
 #include "IDetailCustomization.h"
-
+#include "Input/Reply.h"
 
 class IPropertyHandle;
+class URivermaxMediaOutput;
 
 class FRivermaxMediaDetailsCustomization : public IDetailCustomization
 {
@@ -16,4 +17,9 @@ public:
 	//~ Begin IDetailCustomization Interface
 	virtual void CustomizeDetails(IDetailLayoutBuilder& DetailBuilder) override;
 	//~ End IDetailCustomization Interface
+
+private:
+
+	/** Shows up a path selection dialog and exports SDP there. Currently only works in Media Output and not in ndisplay. */
+	FReply OnExportSDP(URivermaxMediaOutput* Output);
 };
diff --git a/Engine/Plugins/VirtualProduction/Rivermax/RivermaxMedia/Source/RivermaxMediaEditor/Private/ModularFeatures/RivermaxMediaInitializerFeature.cpp b/Engine/Plugins/VirtualProduction/Rivermax/RivermaxMedia/Source/RivermaxMediaEditor/Private/ModularFeatures/RivermaxMediaInitializerFeature.cpp
--- a/Engine/Plugins/VirtualProduction/Rivermax/RivermaxMedia/Source/RivermaxMediaEditor/Private/ModularFeatures/RivermaxMediaInitializerFeature.cpp
+++ b/Engine/Plugins/VirtualProduction/Rivermax/RivermaxMedia/Source/RivermaxMediaEditor/Private/ModularFeatures/RivermaxMediaInitializerFeature.cpp
@@ -72,13 +72,22 @@ void FRivermaxMediaInitializerFeature::InitializeMediaObjectForTile(UObject* Med
 		RivermaxMediaOutput->PresentationQueueSize = 2;
 		RivermaxMediaOutput->bDoFrameCounterTimestamping = true;
 		RivermaxMediaOutput->bOverrideResolution = false;
-		//RivermaxMediaOutput->Resolution          = default value
-		RivermaxMediaOutput->FrameRate           = { 60,1 };
-		RivermaxMediaOutput->PixelFormat         = ERivermaxMediaOutputPixelFormat::PF_10BIT_RGB;
-		RivermaxMediaOutput->InterfaceAddress    = GetRivermaxInterfaceAddress();
-		RivermaxMediaOutput->StreamAddress       = GenerateStreamAddress(OnwerInfo.OwnerUniqueIdx, TilePos);
-		RivermaxMediaOutput->Port                = 50000;
-		RivermaxMediaOutput->bUseGPUDirect       = true;
+		RivermaxMediaOutput->FrameRate = { 60,1 };
+		RivermaxMediaOutput->PixelFormat = ERivermaxMediaOutputPixelFormat::PF_10BIT_RGB;
+		RivermaxMediaOutput->InterfaceAddress = GetRivermaxInterfaceAddress();
+		RivermaxMediaOutput->StreamAddress = GenerateStreamAddress(OnwerInfo.OwnerUniqueIdx, TilePos);
+		RivermaxMediaOutput->Port = 50000;
+		RivermaxMediaOutput->bUseGPUDirect = true;
+
+		/** A setup for UE5.7 */
+		//RivermaxMediaOutput->bOutputVideo = true;
+		//RivermaxMediaOutput->VideoStream.bOverrideResolution = false;
+		//RivermaxMediaOutput->VideoStream.FrameRate           = { 60,1 };
+		//RivermaxMediaOutput->VideoStream.PixelFormat         = ERivermaxMediaOutputPixelFormat::PF_10BIT_RGB;
+		//RivermaxMediaOutput->VideoStream.InterfaceAddress    = GetRivermaxInterfaceAddress();
+		//RivermaxMediaOutput->VideoStream.StreamAddress       = GenerateStreamAddress(OnwerInfo.OwnerUniqueIdx, TilePos);
+		//RivermaxMediaOutput->VideoStream.Port                = 50000;
+		//RivermaxMediaOutput->VideoStream.bUseGPUDirect       = true;
 	}
 }
 
@@ -91,7 +100,6 @@ void FRivermaxMediaInitializerFeature::InitializeMediaObjectForFullFrame(UObject
 		RivermaxMediaSource->bUseTimeSynchronization = true;
 		RivermaxMediaSource->bOverrideResolution = false;
 		RivermaxMediaSource->FrameDelay = 0;
-		//RivermaxMediaSource->Resolution          = default value
 		RivermaxMediaSource->FrameRate           = { 60,1 };
 		RivermaxMediaSource->PixelFormat         = ERivermaxMediaSourcePixelFormat::RGB_10bit;
 		RivermaxMediaSource->InterfaceAddress    = GetRivermaxInterfaceAddress();
@@ -101,32 +109,41 @@ void FRivermaxMediaInitializerFeature::InitializeMediaObjectForFullFrame(UObject
 	}
 	else if (URivermaxMediaOutput* RivermaxMediaOutput = Cast<URivermaxMediaOutput>(MediaObject))
 	{
-		RivermaxMediaOutput->FrameLockingMode      = ERivermaxFrameLockingMode::BlockOnReservation;
+		RivermaxMediaOutput->FrameLockingMode = ERivermaxFrameLockingMode::BlockOnReservation;
 		RivermaxMediaOutput->PresentationQueueSize = 2;
-		RivermaxMediaOutput->bOverrideResolution   = false;
-		//RivermaxMediaOutput->Resolution            = default value
-		RivermaxMediaOutput->PixelFormat           = ERivermaxMediaOutputPixelFormat::PF_10BIT_RGB;
-		RivermaxMediaOutput->InterfaceAddress      = GetRivermaxInterfaceAddress();
-		RivermaxMediaOutput->StreamAddress         = GenerateStreamAddress(OnwerInfo.ClusterNodeUniqueIdx.Get(0), OnwerInfo.OwnerUniqueIdx, OnwerInfo.OwnerType);
-		RivermaxMediaOutput->Port                  = 50000;
+		RivermaxMediaOutput->bOverrideResolution = false;
+		RivermaxMediaOutput->PixelFormat = ERivermaxMediaOutputPixelFormat::PF_10BIT_RGB;
+		RivermaxMediaOutput->InterfaceAddress = GetRivermaxInterfaceAddress();
+		RivermaxMediaOutput->StreamAddress = GenerateStreamAddress(OnwerInfo.ClusterNodeUniqueIdx.Get(0), OnwerInfo.OwnerUniqueIdx, OnwerInfo.OwnerType);
+		RivermaxMediaOutput->Port = 50000;
+
+		//	RivermaxMediaOutput->VideoStream.bOverrideResolution   = false;
+		//	RivermaxMediaOutput->VideoStream.PixelFormat           = ERivermaxMediaOutputPixelFormat::PF_10BIT_RGB;
+		//	RivermaxMediaOutput->VideoStream.InterfaceAddress      = GetRivermaxInterfaceAddress();
+		//	RivermaxMediaOutput->VideoStream.StreamAddress         = GenerateStreamAddress(OnwerInfo.ClusterNodeUniqueIdx.Get(0), OnwerInfo.OwnerUniqueIdx, OnwerInfo.OwnerType);
+		//	RivermaxMediaOutput->VideoStream.Port                  = 50000;
 
 		switch (OnwerInfo.OwnerType)
 		{
 		case FMediaObjectOwnerInfo::EMediaObjectOwnerType::ICVFXCamera:
 		case FMediaObjectOwnerInfo::EMediaObjectOwnerType::Viewport:
-			RivermaxMediaOutput->AlignmentMode       = ERivermaxMediaAlignmentMode::FrameCreation;
+			RivermaxMediaOutput->AlignmentMode = ERivermaxMediaAlignmentMode::FrameCreation;
 			RivermaxMediaOutput->bDoContinuousOutput = false;
 			RivermaxMediaOutput->bDoFrameCounterTimestamping = true;
-			RivermaxMediaOutput->FrameRate           = { 60,1 };
-			RivermaxMediaOutput->bUseGPUDirect       = true;
+			RivermaxMediaOutput->FrameRate = { 60,1 };
+			RivermaxMediaOutput->bUseGPUDirect = true;
+			//RivermaxMediaOutput->VideoStream.FrameRate           = { 60,1 };
+			//RivermaxMediaOutput->VideoStream.bUseGPUDirect       = true;
 			break;
 
 		case FMediaObjectOwnerInfo::EMediaObjectOwnerType::Backbuffer:
-			RivermaxMediaOutput->AlignmentMode       = ERivermaxMediaAlignmentMode::AlignmentPoint;
+			RivermaxMediaOutput->AlignmentMode = ERivermaxMediaAlignmentMode::AlignmentPoint;
 			RivermaxMediaOutput->bDoContinuousOutput = true;
 			RivermaxMediaOutput->bDoFrameCounterTimestamping = false;
-			RivermaxMediaOutput->FrameRate           = { 24,1 };
-			RivermaxMediaOutput->bUseGPUDirect       = false;
+			RivermaxMediaOutput->FrameRate = { 24,1 };
+			RivermaxMediaOutput->bUseGPUDirect = false;
+			//RivermaxMediaOutput->VideoStream.FrameRate           = { 24,1 };
+			//RivermaxMediaOutput->VideoStream.bUseGPUDirect       = false;
 			break;
 
 		default:
diff --git a/Engine/Plugins/VirtualProduction/Rivermax/RivermaxSync/Source/RivermaxSync/Private/MediaOutputSynchronizationPolicyRivermax.cpp b/Engine/Plugins/VirtualProduction/Rivermax/RivermaxSync/Source/RivermaxSync/Private/MediaOutputSynchronizationPolicyRivermax.cpp
--- a/Engine/Plugins/VirtualProduction/Rivermax/RivermaxSync/Source/RivermaxSync/Private/MediaOutputSynchronizationPolicyRivermax.cpp
+++ b/Engine/Plugins/VirtualProduction/Rivermax/RivermaxSync/Source/RivermaxSync/Private/MediaOutputSynchronizationPolicyRivermax.cpp
@@ -71,7 +71,7 @@ bool FMediaOutputSynchronizationPolicyRivermaxHandler::IsCaptureTypeSupported(UM
 
 			if (TimeSource == ERivermaxTimeSource::PTP || TimeSource == ERivermaxTimeSource::System)
 			{
-				const FRivermaxOutputStreamOptions Options = RmaxCapture->GetOutputStreamOptions();
+				const FRivermaxOutputOptions Options = RmaxCapture->GetOutputOptions();
 
 				if (Options.AlignmentMode == ERivermaxAlignmentMode::AlignmentPoint)
 				{
@@ -155,14 +155,23 @@ void FMediaOutputSynchronizationPolicyRivermaxHandler::Synchronize()
 	// Debug cvar to potentially stall a node after exiting the barrier and missing alignment points
 	if (UE::RivermaxSync::GbTriggerRandomDesync)
 	{
+		using namespace UE::RivermaxCore;
+
 		FRandomStream RandomStream(FPlatformTime::Cycles64());
 		const bool bTriggerDesync = (RandomStream.FRandRange(0.0, 1.0) > 0.7) ? true : false;
+
+		URivermaxMediaCapture* RmaxCapture = Cast<URivermaxMediaCapture>(CapturingDevice);
+		FRivermaxOutputOptions OutputOptions = RmaxCapture->GetOutputOptions();
+
+		// Currently we only support synchronization with Video streams.
+		TSharedPtr<FRivermaxVideoOutputOptions> VideoOptions = OutputOptions.GetStreamOptions<FRivermaxVideoOutputOptions>(ERivermaxStreamType::VIDEO_2110_20_STREAM);
+
+		check(VideoOptions.IsValid());
 		if (bTriggerDesync)
 		{
 			TRACE_CPUPROFILER_EVENT_SCOPE(RmaxSync::ForceBadSync);
-			URivermaxMediaCapture* RmaxCapture = Cast<URivermaxMediaCapture>(CapturingDevice);
 			double TimeLeftSeconds = GetTimeBeforeNextSyncPoint();
-			TimeLeftSeconds += RmaxCapture->GetOutputStreamOptions().FrameRate.AsInterval();
+			TimeLeftSeconds += VideoOptions->FrameRate.AsInterval();
 			const float OffsetTimeSeconds = UE::RivermaxSync::CVarRivermaxSyncWakeupOffset.GetValueOnAnyThread() * 1E-3;
 			const float SleepTime = TimeLeftSeconds + OffsetTimeSeconds;
 			FPlatformProcess::SleepNoStats(SleepTime);
@@ -186,8 +195,13 @@ double FMediaOutputSynchronizationPolicyRivermaxHandler::GetTimeBeforeNextSyncPo
 				const uint64 CurrentTimeNanosec = RivermaxMgr->GetTime();
 
 				// Get next alignment timepoint
-				FRivermaxOutputStreamOptions Options = RmaxCapture->GetOutputStreamOptions();
-				const uint64 NextAlignmentTimeNanosec = GetNextAlignmentPoint(CurrentTimeNanosec, Options.FrameRate);
+				FRivermaxOutputOptions Options = RmaxCapture->GetOutputOptions();
+
+				TSharedPtr<FRivermaxVideoOutputOptions> VideoOptions =
+					Options.GetStreamOptions<FRivermaxVideoOutputOptions>(ERivermaxStreamType::VIDEO_2110_20_STREAM);
+
+				check(VideoOptions.IsValid());
+				const uint64 NextAlignmentTimeNanosec = GetNextAlignmentPoint(CurrentTimeNanosec, VideoOptions->FrameRate);
 
 				// Time left
 				checkSlow(NextAlignmentTimeNanosec > CurrentTimeNanosec);
