Rivermax: Media Capture calls GetCustomBufferDescription before Rivermax Media Capture is initialized. During the initialization Rivermax Media Capture was aggregating some properties such as Pixel Format. Because of this GetCustomBufferDescription was using invalid Pixel Format.

#rnx

diff --git a/Engine/Plugins/VirtualProduction/Rivermax/RivermaxCore/Source/RivermaxCore/Private/Streams/RivermaxOutVideoStream.cpp b/Engine/Plugins/VirtualProduction/Rivermax/RivermaxCore/Source/RivermaxCore/Private/Streams/RivermaxOutVideoStream.cpp
--- a/Engine/Plugins/VirtualProduction/Rivermax/RivermaxCore/Source/RivermaxCore/Private/Streams/RivermaxOutVideoStream.cpp
+++ b/Engine/Plugins/VirtualProduction/Rivermax/RivermaxCore/Source/RivermaxCore/Private/Streams/RivermaxOutVideoStream.cpp
@@ -458,17 +458,16 @@ namespace UE::RivermaxCore::Private
 			return false;
 		}
 
-		const SIZE_T FrameSize = GetRowSizeInBytes() * StreamOptions->AlignedResolution.Y;
+		const SIZE_T BytesPerRow = GetRowSizeInBytes();
+		const SIZE_T FrameSize = BytesPerRow * StreamOptions->AlignedResolution.Y;
+
 		if (FrameSize == 0)
 		{
 			UE_LOG(LogRivermax, Warning, TEXT("Can't create Rivermax Output Stream. Frame size of 0 is invalid. Verify resolution."));
 			return false;
 		}
+		
 
-		//2110 stream type
-		//We need to use the fullframe allocated size to compute the payload size.
-
-		const SIZE_T BytesPerRow = GetRowSizeInBytes();
 
 		// By default we want to divide the bytes evenly across packets. Some resolutions will require packets to be sized unevenly.
 		const bool bFoundPayload = FindPayloadSize(BytesPerRow, FormatInfo.PixelGroupSize, StreamMemory.PayloadSize);
