Rivermax:
- Disabled GPU Direct on input by default due to lost packets when GPU and DPU are not on the same root complex.
- Output chunk completion tracker. No longer tracking the last chunk of the frame due to a known Rivermax limitation that the last chunk is completed at the next TRoffset. Instead we are tracking second to last chunk and only if Chunk commit offset optimation is disabled, so that users don't get unexpected waits when polling for chunk completion.

#rnx

diff --git a/Engine/Plugins/VirtualProduction/Rivermax/RivermaxCore/Source/RivermaxCore/Private/RivermaxManager.cpp b/Engine/Plugins/VirtualProduction/Rivermax/RivermaxCore/Source/RivermaxCore/Private/RivermaxManager.cpp
--- a/Engine/Plugins/VirtualProduction/Rivermax/RivermaxCore/Source/RivermaxCore/Private/RivermaxManager.cpp
+++ b/Engine/Plugins/VirtualProduction/Rivermax/RivermaxCore/Source/RivermaxCore/Private/RivermaxManager.cpp
@@ -53,8 +53,9 @@ namespace UE::RivermaxCore::Private
 		ECVF_Default);
 
 	static TAutoConsoleVariable<int32> CVarRivermaxEnableGPUDirectInputCapability(
-		TEXT("Rivermax.GPUDirectInput"), 1,
-		TEXT("Whether GPUDirect capability is enabled for input purposes."),
+		TEXT("Rivermax.GPUDirectInput"), 0,
+		TEXT("Whether GPUDirect capability is enabled for input purposes. When enabled users need to make sure that GPU and DPU are on the same root complex.\n\
+			When streaming multiple Rivermax Inputs using GPU Direct while GPU and DPU are on different root complexes, 2110 packets can be lost. "),
 		ECVF_Default);
 
 	static TAutoConsoleVariable<int32> CVarRivermaxEnableGPUDirectOutputCapability(
diff --git a/Engine/Plugins/VirtualProduction/Rivermax/RivermaxCore/Source/RivermaxCore/Private/Streams/RivermaxOutStream.cpp b/Engine/Plugins/VirtualProduction/Rivermax/RivermaxCore/Source/RivermaxCore/Private/Streams/RivermaxOutStream.cpp
--- a/Engine/Plugins/VirtualProduction/Rivermax/RivermaxCore/Source/RivermaxCore/Private/Streams/RivermaxOutStream.cpp
+++ b/Engine/Plugins/VirtualProduction/Rivermax/RivermaxCore/Source/RivermaxCore/Private/Streams/RivermaxOutStream.cpp
@@ -477,9 +477,17 @@ namespace UE::RivermaxCore::Private
 			Stats.ChunkCompletionTracker->MarkChunkForTracking("First Chunk", StreamData.ChunkHandle, StreamData.NextScheduleTimeNanosec, StreamData.DeltaTimePerChunkNs, CurrentFrame->ChunkNumber);
 		}
 
-		if (CurrentFrame->ChunkNumber == StreamMemory.ChunksPerFrameField - 1)
+		// This is actually tracking second to last chunk. Rivermax API internally doesn't mark the last chunk as completed until after the TRoffset gap, because it is doing some other processing. 
+		if (CurrentFrame->ChunkNumber == StreamMemory.ChunksPerFrameField - 2)
 		{
-			Stats.ChunkCompletionTracker->MarkChunkForTracking("Last Chunk", StreamData.ChunkHandle, StreamData.NextScheduleTimeNanosec, StreamData.DeltaTimePerChunkNs, CurrentFrame->ChunkNumber);
+			// CVarRivermaxOutputCommitChunksOffsetPercent allows the early start of chunk commits. The last chunk will not be completed until the next alignment point. Polling for chunk completion starts right
+			// before the next frame is put on the wire which is at CVarRivermaxOutputCommitChunksOffsetPercent. Polling doesn't exit until it gets chunk completion data, which for the last chunk is right before the next
+			// alignment point. Therefore, this negates the intended effect of CVarRivermaxOutputCommitChunksOffsetPercent, because the commit of the first chunk of the next frame is delayed until 
+			// the last chunk of the current frame is fully completed.
+			if (CVarRivermaxOutputCommitChunksOffsetPercent.GetValueOnAnyThread() < KINDA_SMALL_NUMBER)
+			{
+				Stats.ChunkCompletionTracker->MarkChunkForTracking("Last Chunk", StreamData.ChunkHandle, StreamData.NextScheduleTimeNanosec, StreamData.DeltaTimePerChunkNs, CurrentFrame->ChunkNumber);
+			}
 		}
 
 		do
@@ -1303,6 +1311,8 @@ namespace UE::RivermaxCore::Private
 			return;
 		}
 
+		TRACE_CPUPROFILER_EVENT_SCOPE(FRivermaxChunkCompletionTracker::MarkChunkForTracking);
+
 		// ExpectedChunkCompletionTime is more of a prediction and isn't precise. For that reason adding a margin.
 		constexpr double ChunkCompletionMargin = 1.1; // 10% Margin;
 		const uint64 ExpectedChunkCompletionTime = FrameScheduledTimeNanosec + ((ChunkNumber + 1) * DeltaTimePerChunkNs) + static_cast<uint64>(static_cast<double>(DeltaTimePerChunkNs) * ChunkCompletionMargin);
@@ -1347,6 +1357,7 @@ namespace UE::RivermaxCore::Private
 			do
 			{
 				Status = CachedAPI->rmx_output_media_poll_for_completion(&ChunkHandle);
+				FPlatformProcess::SleepNoStats(0.f);
 			} while (Status != RMX_OK);
 
 			const rmx_output_chunk_completion* ChunkCompletion = CachedAPI->rmx_output_media_get_last_completion(&ChunkHandle);
diff --git a/Engine/Plugins/VirtualProduction/Rivermax/RivermaxCore/Source/RivermaxCore/Private/Streams/RivermaxOutStream.h b/Engine/Plugins/VirtualProduction/Rivermax/RivermaxCore/Source/RivermaxCore/Private/Streams/RivermaxOutStream.h
--- a/Engine/Plugins/VirtualProduction/Rivermax/RivermaxCore/Source/RivermaxCore/Private/Streams/RivermaxOutStream.h
+++ b/Engine/Plugins/VirtualProduction/Rivermax/RivermaxCore/Source/RivermaxCore/Private/Streams/RivermaxOutStream.h
@@ -127,6 +127,10 @@ namespace UE::RivermaxCore::Private
 		* NextScheduleTimeNanosec is the time when the first packet of the frame is supposed to be on the wire.
 		* DeltaTimePerChunk is predicted inter-chunk completion time based on frame rate and number of chunks (FrameRate/NumberOfChunks). In Nanosecs
 		* ChunkNumber is the sequential number of the chunk in the frame.
+		* 
+		* Notes: If you mark chunk for tracking keep in mind that last chunk is completed at TRoffset from the next alignment point due to how it is setup in Rivermax. 
+		* In ideal case scenario the last chunk will complete when the next alignment point starts, so if CVarRivermaxOutputCommitChunksOffsetPercent is used 
+		* the chunk completion polling will eliminate the commit offset by waiting until the next alignment point.
 		*/
 		void MarkChunkForTracking(const FString& ChunkId, const rmx_output_media_chunk_handle& ChunkMarkedForCompletion, uint64 NextScheduleTimeNanosec, uint64 DeltaTimePerChunkNs, uint32 ChunkNumber);
 		
