Rivermax: Removed inline methods and updated rivermax version to 1.71.30

#rnx

diff --git a/Engine/Plugins/VirtualProduction/Rivermax/RivermaxCore/Source/RivermaxCore/Private/RivermaxManager.cpp b/Engine/Plugins/VirtualProduction/Rivermax/RivermaxCore/Source/RivermaxCore/Private/RivermaxManager.cpp
--- a/Engine/Plugins/VirtualProduction/Rivermax/RivermaxCore/Source/RivermaxCore/Private/RivermaxManager.cpp
+++ b/Engine/Plugins/VirtualProduction/Rivermax/RivermaxCore/Source/RivermaxCore/Private/RivermaxManager.cpp
@@ -41,8 +41,8 @@ namespace UE::RivermaxCore
 	TMap<uint8, FString> FRivermaxTracingUtils::RmaxInSelectedFrameTraceEvents;
 
 	// Environment variable used to specify where Rivermax library is located. This enables having an installed version different than the one used by Unreal
-	static const FString RivermaxLibraryEnvironmentVariable = TEXT("RIVERMAX_PATH_1_60_6");
-	static const uint32 RivermaxVersion[3] = { 1, 60, 6 };
+	static const FString RivermaxLibraryEnvironmentVariable = TEXT("RIVERMAX_PATH_1_71_30");
+	static const uint32 RivermaxVersion[3] = { 1, 71, 30 };
 }
 
 namespace UE::RivermaxCore::Private
diff --git a/Engine/Plugins/VirtualProduction/Rivermax/RivermaxCore/Source/RivermaxCore/Private/RivermaxWrapper.cpp b/Engine/Plugins/VirtualProduction/Rivermax/RivermaxCore/Source/RivermaxCore/Private/RivermaxWrapper.cpp
--- a/Engine/Plugins/VirtualProduction/Rivermax/RivermaxCore/Source/RivermaxCore/Private/RivermaxWrapper.cpp
+++ b/Engine/Plugins/VirtualProduction/Rivermax/RivermaxCore/Source/RivermaxCore/Private/RivermaxWrapper.cpp
@@ -18,7 +18,6 @@ namespace UE::RivermaxCore::Private
 
         FunctionList->rmx_get_version_numbers = reinterpret_cast<PFNRMX_GET_VERSION_NUMBERS>(FPlatformProcess::GetDllExport(LibraryHandle, TEXT("rmx_get_version_numbers_v1")));
         FunctionList->rmx_get_version_string = reinterpret_cast<PFNRMX_GET_VERSION_STRING>(FPlatformProcess::GetDllExport(LibraryHandle, TEXT("rmx_get_version_string_v1")));
-        FunctionList->rmx_mark_cpu_for_affinity = reinterpret_cast<PFNRMX_MARK_CPU_FOR_AFFINITY>(FPlatformProcess::GetDllExport(LibraryHandle, TEXT("rmx_mark_cpu_for_affinity")));
         FunctionList->rmx_set_cpu_affinity = reinterpret_cast<PFNRMX_SET_CPU_AFFINITY>(FPlatformProcess::GetDllExport(LibraryHandle, TEXT("rmx_set_cpu_affinity_v1")));
         FunctionList->rmx_enable_system_signal_handling = reinterpret_cast<PFNRMX_ENABLE_SYSTEM_SIGNAL_HANDLING>(FPlatformProcess::GetDllExport(LibraryHandle, TEXT("rmx_enable_system_signal_handling_v1")));
         FunctionList->rmx_init_lib_param = reinterpret_cast<PFNRMX_INIT_LIB_PARAM>(FPlatformProcess::GetDllExport(LibraryHandle, TEXT("rmx_init_lib_param_v1")));
@@ -27,7 +26,6 @@ namespace UE::RivermaxCore::Private
         FunctionList->rmx_set_lib_param_forced = reinterpret_cast<PFNRMX_SET_LIB_PARAM_FORCED>(FPlatformProcess::GetDllExport(LibraryHandle, TEXT("rmx_set_lib_param_forced_v1")));
         FunctionList->rmx_apply_lib_param = reinterpret_cast<PFNRMX_APPLY_LIB_PARAM>(FPlatformProcess::GetDllExport(LibraryHandle, TEXT("rmx_apply_lib_param_v1")));
         FunctionList->_rmx_init = reinterpret_cast<PFN_RMX_INIT>(FPlatformProcess::GetDllExport(LibraryHandle, TEXT("_rmx_init")));
-        FunctionList->rmx_init = reinterpret_cast<PFNRMX_INIT>(FPlatformProcess::GetDllExport(LibraryHandle, TEXT("rmx_init")));
         FunctionList->rmx_cleanup = reinterpret_cast<PFNRMX_CLEANUP>(FPlatformProcess::GetDllExport(LibraryHandle, TEXT("rmx_cleanup_v1")));
         FunctionList->rmx_get_device_list = reinterpret_cast<PFNRMX_GET_DEVICE_LIST>(FPlatformProcess::GetDllExport(LibraryHandle, TEXT("rmx_get_device_list_v1")));
         FunctionList->rmx_free_device_list = reinterpret_cast<PFNRMX_FREE_DEVICE_LIST>(FPlatformProcess::GetDllExport(LibraryHandle, TEXT("rmx_free_device_list_v1")));
@@ -40,13 +38,7 @@ namespace UE::RivermaxCore::Private
         FunctionList->rmx_get_device_id = reinterpret_cast<PFNRMX_GET_DEVICE_ID>(FPlatformProcess::GetDllExport(LibraryHandle, TEXT("rmx_get_device_id_v1")));
         FunctionList->rmx_get_device_serial_number = reinterpret_cast<PFNRMX_GET_DEVICE_SERIAL_NUMBER>(FPlatformProcess::GetDllExport(LibraryHandle, TEXT("rmx_get_device_serial_number_v1")));
         FunctionList->rmx_retrieve_device_iface = reinterpret_cast<PFNRMX_RETRIEVE_DEVICE_IFACE>(FPlatformProcess::GetDllExport(LibraryHandle, TEXT("rmx_retrieve_device_iface_v1")));
-        FunctionList->rmx_retrieve_device_iface_ipv4 = reinterpret_cast<PFNRMX_RETRIEVE_DEVICE_IFACE_IPV4>(FPlatformProcess::GetDllExport(LibraryHandle, TEXT("rmx_retrieve_device_iface_ipv4")));
-        FunctionList->rmx_clear_device_capabilities_enquiry = reinterpret_cast<PFNRMX_CLEAR_DEVICE_CAPABILITIES_ENQUIRY>(FPlatformProcess::GetDllExport(LibraryHandle, TEXT("rmx_clear_device_capabilities_enquiry_v1")));
-        FunctionList->rmx_mark_device_capability_for_enquiry = reinterpret_cast<PFNRMX_MARK_DEVICE_CAPABILITY_FOR_ENQUIRY>(FPlatformProcess::GetDllExport(LibraryHandle, TEXT("rmx_mark_device_capability_for_enquiry")));
         FunctionList->rmx_enquire_device_capabilities = reinterpret_cast<PFNRMX_ENQUIRE_DEVICE_CAPABILITIES>(FPlatformProcess::GetDllExport(LibraryHandle, TEXT("rmx_enquire_device_capabilities_v1")));
-        FunctionList->rmx_is_device_capability_supported = reinterpret_cast<PFNRMX_IS_DEVICE_CAPABILITY_SUPPORTED>(FPlatformProcess::GetDllExport(LibraryHandle, TEXT("rmx_is_device_capability_supported")));
-        FunctionList->rmx_clear_device_config_attributes = reinterpret_cast<PFNRMX_CLEAR_DEVICE_CONFIG_ATTRIBUTES>(FPlatformProcess::GetDllExport(LibraryHandle, TEXT("rmx_clear_device_config_attributes")));
-        FunctionList->rmx_set_device_config_attribute = reinterpret_cast<PFNRMX_SET_DEVICE_CONFIG_ATTRIBUTE>(FPlatformProcess::GetDllExport(LibraryHandle, TEXT("rmx_set_device_config_attribute")));
         FunctionList->rmx_apply_device_config = reinterpret_cast<PFNRMX_APPLY_DEVICE_CONFIG>(FPlatformProcess::GetDllExport(LibraryHandle, TEXT("rmx_apply_device_config_v1")));
         FunctionList->rmx_revert_device_config = reinterpret_cast<PFNRMX_REVERT_DEVICE_CONFIG>(FPlatformProcess::GetDllExport(LibraryHandle, TEXT("rmx_revert_device_config_v1")));
         FunctionList->rmx_init_user_clock = reinterpret_cast<PFNRMX_INIT_USER_CLOCK>(FPlatformProcess::GetDllExport(LibraryHandle, TEXT("rmx_init_user_clock_v1")));
@@ -60,7 +52,6 @@ namespace UE::RivermaxCore::Private
         FunctionList->rmx_get_time = reinterpret_cast<PFNRMX_GET_TIME>(FPlatformProcess::GetDllExport(LibraryHandle, TEXT("rmx_get_time_v1")));
         FunctionList->rmx_init_sockaddr_in = reinterpret_cast<PFNRMX_INIT_SOCKADDR_IN>(FPlatformProcess::GetDllExport(LibraryHandle, TEXT("rmx_init_sockaddr_in_v1")));
         FunctionList->rmx_resolve_sockaddr = reinterpret_cast<PFNRMX_RESOLVE_SOCKADDR>(FPlatformProcess::GetDllExport(LibraryHandle, TEXT("rmx_resolve_sockaddr_v1")));
-        FunctionList->rmx_sockaddr_as_sockaddr = reinterpret_cast<PFNRMX_SOCKADDR_AS_SOCKADDR>(FPlatformProcess::GetDllExport(LibraryHandle, TEXT("rmx_sockaddr_as_sockaddr")));
         FunctionList->rmx_sockaddr_to_std_sockaddr = reinterpret_cast<PFNRMX_SOCKADDR_TO_STD_SOCKADDR>(FPlatformProcess::GetDllExport(LibraryHandle, TEXT("rmx_sockaddr_to_std_sockaddr_v1")));
         FunctionList->rmx_init_event_channel = reinterpret_cast<PFNRMX_INIT_EVENT_CHANNEL>(FPlatformProcess::GetDllExport(LibraryHandle, TEXT("rmx_init_event_channel_v1")));
         FunctionList->rmx_set_event_channel_handle = reinterpret_cast<PFNRMX_SET_EVENT_CHANNEL_HANDLE>(FPlatformProcess::GetDllExport(LibraryHandle, TEXT("rmx_set_event_channel_handle_v1")));
@@ -78,7 +69,6 @@ namespace UE::RivermaxCore::Private
         FunctionList->rmx_input_enable_stream_option = reinterpret_cast<PFNRMX_INPUT_ENABLE_STREAM_OPTION>(FPlatformProcess::GetDllExport(LibraryHandle, TEXT("rmx_input_enable_stream_option_v1")));
         FunctionList->rmx_input_set_mem_sub_block_count = reinterpret_cast<PFNRMX_INPUT_SET_MEM_SUB_BLOCK_COUNT>(FPlatformProcess::GetDllExport(LibraryHandle, TEXT("rmx_input_set_mem_sub_block_count_v1")));
         FunctionList->rmx_input_set_entry_size_range = reinterpret_cast<PFNRMX_INPUT_SET_ENTRY_SIZE_RANGE>(FPlatformProcess::GetDllExport(LibraryHandle, TEXT("rmx_input_set_entry_size_range_v1")));
-        FunctionList->rmx_input_set_entry_uniform_size = reinterpret_cast<PFNRMX_INPUT_SET_ENTRY_UNIFORM_SIZE>(FPlatformProcess::GetDllExport(LibraryHandle, TEXT("rmx_input_set_entry_uniform_size")));
         FunctionList->rmx_input_determine_mem_layout = reinterpret_cast<PFNRMX_INPUT_DETERMINE_MEM_LAYOUT>(FPlatformProcess::GetDllExport(LibraryHandle, TEXT("rmx_input_determine_mem_layout_v1")));
         FunctionList->rmx_input_get_mem_capacity_in_packets = reinterpret_cast<PFNRMX_INPUT_GET_MEM_CAPACITY_IN_PACKETS>(FPlatformProcess::GetDllExport(LibraryHandle, TEXT("rmx_input_get_mem_capacity_in_packets_v1")));
         FunctionList->rmx_input_get_mem_block_buffer = reinterpret_cast<PFNRMX_INPUT_GET_MEM_BLOCK_BUFFER>(FPlatformProcess::GetDllExport(LibraryHandle, TEXT("rmx_input_get_mem_block_buffer_v1")));
@@ -96,16 +86,7 @@ namespace UE::RivermaxCore::Private
         FunctionList->rmx_input_init_chunk_handle = reinterpret_cast<PFNRMX_INPUT_INIT_CHUNK_HANDLE>(FPlatformProcess::GetDllExport(LibraryHandle, TEXT("rmx_input_init_chunk_handle_v1")));
         FunctionList->rmx_input_get_next_chunk = reinterpret_cast<PFNRMX_INPUT_GET_NEXT_CHUNK>(FPlatformProcess::GetDllExport(LibraryHandle, TEXT("rmx_input_get_next_chunk_v2")));
         FunctionList->rmx_input_get_chunk_completion = reinterpret_cast<PFNRMX_INPUT_GET_CHUNK_COMPLETION>(FPlatformProcess::GetDllExport(LibraryHandle, TEXT("rmx_input_get_chunk_completion_v1")));
-        FunctionList->rmx_input_get_completion_ptr = reinterpret_cast<PFNRMX_INPUT_GET_COMPLETION_PTR>(FPlatformProcess::GetDllExport(LibraryHandle, TEXT("rmx_input_get_completion_ptr")));
-        FunctionList->rmx_input_get_completion_chunk_size = reinterpret_cast<PFNRMX_INPUT_GET_COMPLETION_CHUNK_SIZE>(FPlatformProcess::GetDllExport(LibraryHandle, TEXT("rmx_input_get_completion_chunk_size")));
-        FunctionList->rmx_input_get_completion_seqn_first = reinterpret_cast<PFNRMX_INPUT_GET_COMPLETION_SEQN_FIRST>(FPlatformProcess::GetDllExport(LibraryHandle, TEXT("rmx_input_get_completion_seqn_first")));
-        FunctionList->rmx_input_get_completion_flag = reinterpret_cast<PFNRMX_INPUT_GET_COMPLETION_FLAG>(FPlatformProcess::GetDllExport(LibraryHandle, TEXT("rmx_input_get_completion_flag")));
-        FunctionList->rmx_input_get_completion_timestamp_first = reinterpret_cast<PFNRMX_INPUT_GET_COMPLETION_TIMESTAMP_FIRST>(FPlatformProcess::GetDllExport(LibraryHandle, TEXT("rmx_input_get_completion_timestamp_first")));
-        FunctionList->rmx_input_get_completion_timestamp_last = reinterpret_cast<PFNRMX_INPUT_GET_COMPLETION_TIMESTAMP_LAST>(FPlatformProcess::GetDllExport(LibraryHandle, TEXT("rmx_input_get_completion_timestamp_last")));
         FunctionList->rmx_input_get_packet_info = reinterpret_cast<PFNRMX_INPUT_GET_PACKET_INFO>(FPlatformProcess::GetDllExport(LibraryHandle, TEXT("rmx_input_get_packet_info_v1")));
-        FunctionList->rmx_input_get_packet_size = reinterpret_cast<PFNRMX_INPUT_GET_PACKET_SIZE>(FPlatformProcess::GetDllExport(LibraryHandle, TEXT("rmx_input_get_packet_size")));
-        FunctionList->rmx_input_get_packet_flow_tag = reinterpret_cast<PFNRMX_INPUT_GET_PACKET_FLOW_TAG>(FPlatformProcess::GetDllExport(LibraryHandle, TEXT("rmx_input_get_packet_flow_tag")));
-        FunctionList->rmx_input_get_packet_timestamp = reinterpret_cast<PFNRMX_INPUT_GET_PACKET_TIMESTAMP>(FPlatformProcess::GetDllExport(LibraryHandle, TEXT("rmx_input_get_packet_timestamp")));
         FunctionList->rmx_output_media_init = reinterpret_cast<PFNRMX_OUTPUT_MEDIA_INIT>(FPlatformProcess::GetDllExport(LibraryHandle, TEXT("rmx_output_media_init_v1")));
         FunctionList->rmx_output_media_set_sdp = reinterpret_cast<PFNRMX_OUTPUT_MEDIA_SET_SDP>(FPlatformProcess::GetDllExport(LibraryHandle, TEXT("rmx_output_media_set_sdp_v1")));
         FunctionList->rmx_output_media_set_idx_in_sdp = reinterpret_cast<PFNRMX_OUTPUT_MEDIA_SET_IDX_IN_SDP>(FPlatformProcess::GetDllExport(LibraryHandle, TEXT("rmx_output_media_set_idx_in_sdp_v1")));
@@ -133,11 +114,6 @@ namespace UE::RivermaxCore::Private
         FunctionList->rmx_output_media_init_chunk_handle = reinterpret_cast<PFNRMX_OUTPUT_MEDIA_INIT_CHUNK_HANDLE>(FPlatformProcess::GetDllExport(LibraryHandle, TEXT("rmx_output_media_init_chunk_handle_v1")));
         FunctionList->rmx_output_media_set_chunk_packet_count = reinterpret_cast<PFNRMX_OUTPUT_MEDIA_SET_CHUNK_PACKET_COUNT>(FPlatformProcess::GetDllExport(LibraryHandle, TEXT("rmx_output_media_set_chunk_packet_count_v1")));
         FunctionList->rmx_output_media_get_next_chunk = reinterpret_cast<PFNRMX_OUTPUT_MEDIA_GET_NEXT_CHUNK>(FPlatformProcess::GetDllExport(LibraryHandle, TEXT("rmx_output_media_get_next_chunk_v1")));
-        FunctionList->rmx_output_media_get_chunk_strides = reinterpret_cast<PFNRMX_OUTPUT_MEDIA_GET_CHUNK_STRIDES>(FPlatformProcess::GetDllExport(LibraryHandle, TEXT("rmx_output_media_get_chunk_strides")));
-        FunctionList->rmx_output_media_get_chunk_packet_sizes = reinterpret_cast<PFNRMX_OUTPUT_MEDIA_GET_CHUNK_PACKET_SIZES>(FPlatformProcess::GetDllExport(LibraryHandle, TEXT("rmx_output_media_get_chunk_packet_sizes")));
-        FunctionList->rmx_output_media_set_chunk_option = reinterpret_cast<PFNRMX_OUTPUT_MEDIA_SET_CHUNK_OPTION>(FPlatformProcess::GetDllExport(LibraryHandle, TEXT("rmx_output_media_set_chunk_option")));
-        FunctionList->rmx_output_media_clear_chunk_option = reinterpret_cast<PFNRMX_OUTPUT_MEDIA_CLEAR_CHUNK_OPTION>(FPlatformProcess::GetDllExport(LibraryHandle, TEXT("rmx_output_media_clear_chunk_option")));
-        FunctionList->rmx_output_media_clear_chunk_all_options = reinterpret_cast<PFNRMX_OUTPUT_MEDIA_CLEAR_CHUNK_ALL_OPTIONS>(FPlatformProcess::GetDllExport(LibraryHandle, TEXT("rmx_output_media_clear_chunk_all_options")));
         FunctionList->rmx_output_media_mark_chunk_for_tracking = reinterpret_cast<PFNRMX_OUTPUT_MEDIA_MARK_CHUNK_FOR_TRACKING>(FPlatformProcess::GetDllExport(LibraryHandle, TEXT("rmx_output_media_mark_chunk_for_tracking_v1")));
         FunctionList->rmx_output_media_commit_chunk = reinterpret_cast<PFNRMX_OUTPUT_MEDIA_COMMIT_CHUNK>(FPlatformProcess::GetDllExport(LibraryHandle, TEXT("rmx_output_media_commit_chunk_v1")));
         FunctionList->rmx_output_media_poll_for_completion = reinterpret_cast<PFNRMX_OUTPUT_MEDIA_POLL_FOR_COMPLETION>(FPlatformProcess::GetDllExport(LibraryHandle, TEXT("rmx_output_media_poll_for_completion_v1")));
@@ -163,9 +139,6 @@ namespace UE::RivermaxCore::Private
         FunctionList->rmx_output_gen_init_chunk_handle = reinterpret_cast<PFNRMX_OUTPUT_GEN_INIT_CHUNK_HANDLE>(FPlatformProcess::GetDllExport(LibraryHandle, TEXT("rmx_output_gen_init_chunk_handle_v1")));
         FunctionList->rmx_output_gen_get_next_chunk = reinterpret_cast<PFNRMX_OUTPUT_GEN_GET_NEXT_CHUNK>(FPlatformProcess::GetDllExport(LibraryHandle, TEXT("rmx_output_gen_get_next_chunk_v1")));
         FunctionList->rmx_output_gen_set_chunk_remote_addr = reinterpret_cast<PFNRMX_OUTPUT_GEN_SET_CHUNK_REMOTE_ADDR>(FPlatformProcess::GetDllExport(LibraryHandle, TEXT("rmx_output_gen_set_chunk_remote_addr_v1")));
-        FunctionList->rmx_output_gen_set_chunk_option = reinterpret_cast<PFNRMX_OUTPUT_GEN_SET_CHUNK_OPTION>(FPlatformProcess::GetDllExport(LibraryHandle, TEXT("rmx_output_gen_set_chunk_option")));
-        FunctionList->rmx_output_gen_clear_chunk_option = reinterpret_cast<PFNRMX_OUTPUT_GEN_CLEAR_CHUNK_OPTION>(FPlatformProcess::GetDllExport(LibraryHandle, TEXT("rmx_output_gen_clear_chunk_option")));
-        FunctionList->rmx_output_gen_clear_chunk_all_options = reinterpret_cast<PFNRMX_OUTPUT_GEN_CLEAR_CHUNK_ALL_OPTIONS>(FPlatformProcess::GetDllExport(LibraryHandle, TEXT("rmx_output_gen_clear_chunk_all_options")));
         FunctionList->rmx_output_gen_append_packet_to_chunk = reinterpret_cast<PFNRMX_OUTPUT_GEN_APPEND_PACKET_TO_CHUNK>(FPlatformProcess::GetDllExport(LibraryHandle, TEXT("rmx_output_gen_append_packet_to_chunk_v1")));
         FunctionList->rmx_output_gen_mark_chunk_for_tracking = reinterpret_cast<PFNRMX_OUTPUT_GEN_MARK_CHUNK_FOR_TRACKING>(FPlatformProcess::GetDllExport(LibraryHandle, TEXT("rmx_output_gen_mark_chunk_for_tracking_v1")));
         FunctionList->rmx_output_gen_commit_chunk = reinterpret_cast<PFNRMX_OUTPUT_GEN_COMMIT_CHUNK>(FPlatformProcess::GetDllExport(LibraryHandle, TEXT("rmx_output_gen_commit_chunk_v1")));
@@ -174,10 +147,8 @@ namespace UE::RivermaxCore::Private
         FunctionList->rmx_output_get_chunk_count = reinterpret_cast<PFNRMX_OUTPUT_GET_CHUNK_COUNT>(FPlatformProcess::GetDllExport(LibraryHandle, TEXT("rmx_output_get_chunk_count_v1")));
         FunctionList->rmx_output_update_dscp = reinterpret_cast<PFNRMX_OUTPUT_UPDATE_DSCP>(FPlatformProcess::GetDllExport(LibraryHandle, TEXT("rmx_output_update_dscp_v1")));
         FunctionList->rmx_output_update_ecn = reinterpret_cast<PFNRMX_OUTPUT_UPDATE_ECN>(FPlatformProcess::GetDllExport(LibraryHandle, TEXT("rmx_output_update_ecn_v1")));
-        FunctionList->rmx_output_get_completion_user_token = reinterpret_cast<PFNRMX_OUTPUT_GET_COMPLETION_USER_TOKEN>(FPlatformProcess::GetDllExport(LibraryHandle, TEXT("rmx_output_get_completion_user_token")));
-        FunctionList->rmx_output_get_completion_timestamp = reinterpret_cast<PFNRMX_OUTPUT_GET_COMPLETION_TIMESTAMP>(FPlatformProcess::GetDllExport(LibraryHandle, TEXT("rmx_output_get_completion_timestamp")));
 
-        // Check if all functions were loaded succesfully.
+        // Check if all functions were loaded successfully.
         bool bHasLoadedFunctions = true;
         bHasLoadedFunctions = bHasLoadedFunctions && FunctionList->rmx_get_version_numbers;
         bHasLoadedFunctions = bHasLoadedFunctions && FunctionList->rmx_get_version_string;
diff --git a/Engine/Plugins/VirtualProduction/Rivermax/RivermaxCore/Source/RivermaxCore/Private/RivermaxWrapper.h b/Engine/Plugins/VirtualProduction/Rivermax/RivermaxCore/Source/RivermaxCore/Private/RivermaxWrapper.h
--- a/Engine/Plugins/VirtualProduction/Rivermax/RivermaxCore/Source/RivermaxCore/Private/RivermaxWrapper.h
+++ b/Engine/Plugins/VirtualProduction/Rivermax/RivermaxCore/Source/RivermaxCore/Private/RivermaxWrapper.h
@@ -11,7 +11,6 @@ namespace UE::RivermaxCore::Private
 
     typedef const rmx_version * (*PFNRMX_GET_VERSION_NUMBERS)(void);
     typedef const char * (*PFNRMX_GET_VERSION_STRING)(void);
-    typedef void (*PFNRMX_MARK_CPU_FOR_AFFINITY)(uint64_t * bitmask, size_t core_number);
     typedef rmx_status (*PFNRMX_SET_CPU_AFFINITY)(const uint64_t * bitmask, size_t core_count);
     typedef rmx_status (*PFNRMX_ENABLE_SYSTEM_SIGNAL_HANDLING)(void);
     typedef void (*PFNRMX_INIT_LIB_PARAM)(rmx_lib_param * descriptor);
@@ -20,7 +19,6 @@ namespace UE::RivermaxCore::Private
     typedef void (*PFNRMX_SET_LIB_PARAM_FORCED)(rmx_lib_param * descriptor);
     typedef rmx_status (*PFNRMX_APPLY_LIB_PARAM)(rmx_lib_param * descriptor);
     typedef rmx_status (*PFN_RMX_INIT)(const rmx_version * policy);
-    typedef rmx_status (*PFNRMX_INIT)(void);
     typedef rmx_status (*PFNRMX_CLEANUP)(void);
     typedef size_t (*PFNRMX_GET_DEVICE_LIST)(rmx_device_list ** list);
     typedef void (*PFNRMX_FREE_DEVICE_LIST)(rmx_device_list * list);
@@ -33,13 +31,8 @@ namespace UE::RivermaxCore::Private
     typedef uint32_t (*PFNRMX_GET_DEVICE_ID)(const rmx_device * device);
     typedef const char * (*PFNRMX_GET_DEVICE_SERIAL_NUMBER)(const rmx_device * device);
     typedef rmx_status (*PFNRMX_RETRIEVE_DEVICE_IFACE)(rmx_device_iface * device_iface, const rmx_ip_addr * ip);
-    typedef rmx_status (*PFNRMX_RETRIEVE_DEVICE_IFACE_IPV4)(rmx_device_iface * device_iface, const struct in_addr * ip);
     typedef void (*PFNRMX_CLEAR_DEVICE_CAPABILITIES_ENQUIRY)(rmx_device_capabilities * caps);
-    typedef void (*PFNRMX_MARK_DEVICE_CAPABILITY_FOR_ENQUIRY)(rmx_device_capabilities * caps, rmx_device_capability capability);
     typedef rmx_status (*PFNRMX_ENQUIRE_DEVICE_CAPABILITIES)(const rmx_device_iface * device_iface, rmx_device_capabilities * caps);
-    typedef bool (*PFNRMX_IS_DEVICE_CAPABILITY_SUPPORTED)(const rmx_device_capabilities * caps, rmx_device_capability capability);
-    typedef void (*PFNRMX_CLEAR_DEVICE_CONFIG_ATTRIBUTES)(rmx_device_config * config);
-    typedef void (*PFNRMX_SET_DEVICE_CONFIG_ATTRIBUTE)(rmx_device_config * config, rmx_device_config_attribute attribute);
     typedef rmx_status (*PFNRMX_APPLY_DEVICE_CONFIG)(const rmx_device_iface * device_iface, rmx_device_config * config);
     typedef rmx_status (*PFNRMX_REVERT_DEVICE_CONFIG)(const rmx_device_iface * device_iface, rmx_device_config * config);
     typedef void (*PFNRMX_INIT_USER_CLOCK)(rmx_user_clock_params * params);
@@ -53,7 +46,6 @@ namespace UE::RivermaxCore::Private
     typedef rmx_status (*PFNRMX_GET_TIME)(rmx_time_type type, uint64_t * time);
     typedef rmx_status (*PFNRMX_INIT_SOCKADDR_IN)(rmx_sockaddr * addr, const struct sockaddr_in * src);
     typedef rmx_status (*PFNRMX_RESOLVE_SOCKADDR)(rmx_sockaddr * addr, const rmx_device_iface * device_iface);
-    typedef struct sockaddr * (*PFNRMX_SOCKADDR_AS_SOCKADDR)(const rmx_sockaddr * addr);
     typedef rmx_status (*PFNRMX_SOCKADDR_TO_STD_SOCKADDR)(struct sockaddr * dst, const rmx_sockaddr * src);
     typedef void (*PFNRMX_INIT_EVENT_CHANNEL)(rmx_event_channel_params * params, rmx_stream_id id);
     typedef void (*PFNRMX_SET_EVENT_CHANNEL_HANDLE)(rmx_event_channel_params * params, rmx_event_channel_handle * handle);
@@ -71,7 +63,6 @@ namespace UE::RivermaxCore::Private
     typedef void (*PFNRMX_INPUT_ENABLE_STREAM_OPTION)(rmx_input_stream_params * params, rmx_input_option option);
     typedef void (*PFNRMX_INPUT_SET_MEM_SUB_BLOCK_COUNT)(rmx_input_stream_params * params, size_t count);
     typedef void (*PFNRMX_INPUT_SET_ENTRY_SIZE_RANGE)(rmx_input_stream_params * params, size_t sub_block_id, size_t min_size, size_t max_size);
-    typedef void (*PFNRMX_INPUT_SET_ENTRY_UNIFORM_SIZE)(rmx_input_stream_params * params, size_t sub_block_id, size_t size);
     typedef rmx_status (*PFNRMX_INPUT_DETERMINE_MEM_LAYOUT)(rmx_input_stream_params * params);
     typedef size_t (*PFNRMX_INPUT_GET_MEM_CAPACITY_IN_PACKETS)(const rmx_input_stream_params * params);
     typedef rmx_mem_region * (*PFNRMX_INPUT_GET_MEM_BLOCK_BUFFER)(rmx_input_stream_params * params, size_t sub_block_id);
@@ -89,16 +80,7 @@ namespace UE::RivermaxCore::Private
     typedef void (*PFNRMX_INPUT_INIT_CHUNK_HANDLE)(rmx_input_chunk_handle * handle, rmx_stream_id id);
     typedef rmx_status (*PFNRMX_INPUT_GET_NEXT_CHUNK)(rmx_input_chunk_handle * handle);
     typedef const rmx_input_completion * (*PFNRMX_INPUT_GET_CHUNK_COMPLETION)(rmx_input_chunk_handle * handle);
-    typedef const void * (*PFNRMX_INPUT_GET_COMPLETION_PTR)(const rmx_input_completion * completion, size_t sub_block_id);
-    typedef size_t (*PFNRMX_INPUT_GET_COMPLETION_CHUNK_SIZE)(const rmx_input_completion * completion);
-    typedef uint32_t (*PFNRMX_INPUT_GET_COMPLETION_SEQN_FIRST)(const rmx_input_completion * completion);
-    typedef bool (*PFNRMX_INPUT_GET_COMPLETION_FLAG)(const rmx_input_completion * completion, rmx_input_completion_flag flag);
-    typedef uint64_t (*PFNRMX_INPUT_GET_COMPLETION_TIMESTAMP_FIRST)(const rmx_input_completion * completion);
-    typedef uint64_t (*PFNRMX_INPUT_GET_COMPLETION_TIMESTAMP_LAST)(const rmx_input_completion * completion);
     typedef const rmx_input_packet_info * (*PFNRMX_INPUT_GET_PACKET_INFO)(rmx_input_chunk_handle * handle, size_t packet_id);
-    typedef size_t (*PFNRMX_INPUT_GET_PACKET_SIZE)(const rmx_input_packet_info * info, size_t sub_block_id);
-    typedef uint32_t (*PFNRMX_INPUT_GET_PACKET_FLOW_TAG)(const rmx_input_packet_info * info);
-    typedef uint64_t (*PFNRMX_INPUT_GET_PACKET_TIMESTAMP)(const rmx_input_packet_info * info);
     typedef void (*PFNRMX_OUTPUT_MEDIA_INIT)(rmx_output_media_stream_params * params);
     typedef void (*PFNRMX_OUTPUT_MEDIA_SET_SDP)(rmx_output_media_stream_params * params, const char * sdp);
     typedef void (*PFNRMX_OUTPUT_MEDIA_SET_IDX_IN_SDP)(rmx_output_media_stream_params * params, size_t media_block_idx);
@@ -126,11 +108,6 @@ namespace UE::RivermaxCore::Private
     typedef void (*PFNRMX_OUTPUT_MEDIA_INIT_CHUNK_HANDLE)(rmx_output_media_chunk_handle * handle, rmx_stream_id id);
     typedef void (*PFNRMX_OUTPUT_MEDIA_SET_CHUNK_PACKET_COUNT)(rmx_output_media_chunk_handle * handle, size_t packets_in_chunk);
     typedef rmx_status (*PFNRMX_OUTPUT_MEDIA_GET_NEXT_CHUNK)(rmx_output_media_chunk_handle * handle);
-    typedef void * (*PFNRMX_OUTPUT_MEDIA_GET_CHUNK_STRIDES)(rmx_output_media_chunk_handle * handle, size_t sub_block_id);
-    typedef uint16_t * (*PFNRMX_OUTPUT_MEDIA_GET_CHUNK_PACKET_SIZES)(rmx_output_media_chunk_handle * handle, size_t sub_block_id);
-    typedef void (*PFNRMX_OUTPUT_MEDIA_SET_CHUNK_OPTION)(const rmx_output_media_chunk_handle * handle, rmx_output_commit_option option);
-    typedef void (*PFNRMX_OUTPUT_MEDIA_CLEAR_CHUNK_OPTION)(const rmx_output_media_chunk_handle * handle, rmx_output_commit_option option);
-    typedef void (*PFNRMX_OUTPUT_MEDIA_CLEAR_CHUNK_ALL_OPTIONS)(const rmx_output_media_chunk_handle * handle);
     typedef rmx_status (*PFNRMX_OUTPUT_MEDIA_MARK_CHUNK_FOR_TRACKING)(const rmx_output_media_chunk_handle * handle, uint64_t token);
     typedef rmx_status (*PFNRMX_OUTPUT_MEDIA_COMMIT_CHUNK)(const rmx_output_media_chunk_handle * handle, uint64_t time);
     typedef rmx_status (*PFNRMX_OUTPUT_MEDIA_POLL_FOR_COMPLETION)(const rmx_output_media_chunk_handle * handle);
@@ -156,9 +133,6 @@ namespace UE::RivermaxCore::Private
     typedef void (*PFNRMX_OUTPUT_GEN_INIT_CHUNK_HANDLE)(rmx_output_gen_chunk_handle * handle, rmx_stream_id id);
     typedef rmx_status (*PFNRMX_OUTPUT_GEN_GET_NEXT_CHUNK)(rmx_output_gen_chunk_handle * handle);
     typedef void (*PFNRMX_OUTPUT_GEN_SET_CHUNK_REMOTE_ADDR)(rmx_output_gen_chunk_handle * handle, const struct sockaddr * addr);
-    typedef void (*PFNRMX_OUTPUT_GEN_SET_CHUNK_OPTION)(rmx_output_gen_chunk_handle * handle, rmx_output_commit_option option);
-    typedef void (*PFNRMX_OUTPUT_GEN_CLEAR_CHUNK_OPTION)(rmx_output_gen_chunk_handle * handle, rmx_output_commit_option option);
-    typedef void (*PFNRMX_OUTPUT_GEN_CLEAR_CHUNK_ALL_OPTIONS)(rmx_output_gen_chunk_handle * handle);
     typedef rmx_status (*PFNRMX_OUTPUT_GEN_APPEND_PACKET_TO_CHUNK)(rmx_output_gen_chunk_handle * handle, const rmx_mem_region * sub_blocks, size_t count);
     typedef rmx_status (*PFNRMX_OUTPUT_GEN_MARK_CHUNK_FOR_TRACKING)(const rmx_output_gen_chunk_handle * handle, uint64_t token);
     typedef rmx_status (*PFNRMX_OUTPUT_GEN_COMMIT_CHUNK)(rmx_output_gen_chunk_handle * handle, uint64_t time);
@@ -167,14 +141,12 @@ namespace UE::RivermaxCore::Private
     typedef rmx_status (*PFNRMX_OUTPUT_GET_CHUNK_COUNT)(rmx_stream_id id, size_t * count);
     typedef rmx_status (*PFNRMX_OUTPUT_UPDATE_DSCP)(rmx_stream_id id, uint8_t dscp);
     typedef rmx_status (*PFNRMX_OUTPUT_UPDATE_ECN)(rmx_stream_id id, uint8_t ecn);
-    typedef uint64_t (*PFNRMX_OUTPUT_GET_COMPLETION_USER_TOKEN)(const rmx_output_chunk_completion * completion);
     typedef uint64_t (*PFNRMX_OUTPUT_GET_COMPLETION_TIMESTAMP)(const rmx_output_chunk_completion * completion);
 
     struct RIVERMAX_API_FUNCTION_LIST
     {
         PFNRMX_GET_VERSION_NUMBERS rmx_get_version_numbers;
         PFNRMX_GET_VERSION_STRING rmx_get_version_string;
-        PFNRMX_MARK_CPU_FOR_AFFINITY rmx_mark_cpu_for_affinity;
         PFNRMX_SET_CPU_AFFINITY rmx_set_cpu_affinity;
         PFNRMX_ENABLE_SYSTEM_SIGNAL_HANDLING rmx_enable_system_signal_handling;
         PFNRMX_INIT_LIB_PARAM rmx_init_lib_param;
@@ -183,7 +155,6 @@ namespace UE::RivermaxCore::Private
         PFNRMX_SET_LIB_PARAM_FORCED rmx_set_lib_param_forced;
         PFNRMX_APPLY_LIB_PARAM rmx_apply_lib_param;
         PFN_RMX_INIT _rmx_init;
-        PFNRMX_INIT rmx_init;
         PFNRMX_CLEANUP rmx_cleanup;
         PFNRMX_GET_DEVICE_LIST rmx_get_device_list;
         PFNRMX_FREE_DEVICE_LIST rmx_free_device_list;
@@ -196,13 +167,8 @@ namespace UE::RivermaxCore::Private
         PFNRMX_GET_DEVICE_ID rmx_get_device_id;
         PFNRMX_GET_DEVICE_SERIAL_NUMBER rmx_get_device_serial_number;
         PFNRMX_RETRIEVE_DEVICE_IFACE rmx_retrieve_device_iface;
-        PFNRMX_RETRIEVE_DEVICE_IFACE_IPV4 rmx_retrieve_device_iface_ipv4;
         PFNRMX_CLEAR_DEVICE_CAPABILITIES_ENQUIRY rmx_clear_device_capabilities_enquiry;
-        PFNRMX_MARK_DEVICE_CAPABILITY_FOR_ENQUIRY rmx_mark_device_capability_for_enquiry;
         PFNRMX_ENQUIRE_DEVICE_CAPABILITIES rmx_enquire_device_capabilities;
-        PFNRMX_IS_DEVICE_CAPABILITY_SUPPORTED rmx_is_device_capability_supported;
-        PFNRMX_CLEAR_DEVICE_CONFIG_ATTRIBUTES rmx_clear_device_config_attributes;
-        PFNRMX_SET_DEVICE_CONFIG_ATTRIBUTE rmx_set_device_config_attribute;
         PFNRMX_APPLY_DEVICE_CONFIG rmx_apply_device_config;
         PFNRMX_REVERT_DEVICE_CONFIG rmx_revert_device_config;
         PFNRMX_INIT_USER_CLOCK rmx_init_user_clock;
@@ -216,7 +182,6 @@ namespace UE::RivermaxCore::Private
         PFNRMX_GET_TIME rmx_get_time;
         PFNRMX_INIT_SOCKADDR_IN rmx_init_sockaddr_in;
         PFNRMX_RESOLVE_SOCKADDR rmx_resolve_sockaddr;
-        PFNRMX_SOCKADDR_AS_SOCKADDR rmx_sockaddr_as_sockaddr;
         PFNRMX_SOCKADDR_TO_STD_SOCKADDR rmx_sockaddr_to_std_sockaddr;
         PFNRMX_INIT_EVENT_CHANNEL rmx_init_event_channel;
         PFNRMX_SET_EVENT_CHANNEL_HANDLE rmx_set_event_channel_handle;
@@ -234,7 +199,6 @@ namespace UE::RivermaxCore::Private
         PFNRMX_INPUT_ENABLE_STREAM_OPTION rmx_input_enable_stream_option;
         PFNRMX_INPUT_SET_MEM_SUB_BLOCK_COUNT rmx_input_set_mem_sub_block_count;
         PFNRMX_INPUT_SET_ENTRY_SIZE_RANGE rmx_input_set_entry_size_range;
-        PFNRMX_INPUT_SET_ENTRY_UNIFORM_SIZE rmx_input_set_entry_uniform_size;
         PFNRMX_INPUT_DETERMINE_MEM_LAYOUT rmx_input_determine_mem_layout;
         PFNRMX_INPUT_GET_MEM_CAPACITY_IN_PACKETS rmx_input_get_mem_capacity_in_packets;
         PFNRMX_INPUT_GET_MEM_BLOCK_BUFFER rmx_input_get_mem_block_buffer;
@@ -252,16 +216,7 @@ namespace UE::RivermaxCore::Private
         PFNRMX_INPUT_INIT_CHUNK_HANDLE rmx_input_init_chunk_handle;
         PFNRMX_INPUT_GET_NEXT_CHUNK rmx_input_get_next_chunk;
         PFNRMX_INPUT_GET_CHUNK_COMPLETION rmx_input_get_chunk_completion;
-        PFNRMX_INPUT_GET_COMPLETION_PTR rmx_input_get_completion_ptr;
-        PFNRMX_INPUT_GET_COMPLETION_CHUNK_SIZE rmx_input_get_completion_chunk_size;
-        PFNRMX_INPUT_GET_COMPLETION_SEQN_FIRST rmx_input_get_completion_seqn_first;
-        PFNRMX_INPUT_GET_COMPLETION_FLAG rmx_input_get_completion_flag;
-        PFNRMX_INPUT_GET_COMPLETION_TIMESTAMP_FIRST rmx_input_get_completion_timestamp_first;
-        PFNRMX_INPUT_GET_COMPLETION_TIMESTAMP_LAST rmx_input_get_completion_timestamp_last;
         PFNRMX_INPUT_GET_PACKET_INFO rmx_input_get_packet_info;
-        PFNRMX_INPUT_GET_PACKET_SIZE rmx_input_get_packet_size;
-        PFNRMX_INPUT_GET_PACKET_FLOW_TAG rmx_input_get_packet_flow_tag;
-        PFNRMX_INPUT_GET_PACKET_TIMESTAMP rmx_input_get_packet_timestamp;
         PFNRMX_OUTPUT_MEDIA_INIT rmx_output_media_init;
         PFNRMX_OUTPUT_MEDIA_SET_SDP rmx_output_media_set_sdp;
         PFNRMX_OUTPUT_MEDIA_SET_IDX_IN_SDP rmx_output_media_set_idx_in_sdp;
@@ -289,11 +244,6 @@ namespace UE::RivermaxCore::Private
         PFNRMX_OUTPUT_MEDIA_INIT_CHUNK_HANDLE rmx_output_media_init_chunk_handle;
         PFNRMX_OUTPUT_MEDIA_SET_CHUNK_PACKET_COUNT rmx_output_media_set_chunk_packet_count;
         PFNRMX_OUTPUT_MEDIA_GET_NEXT_CHUNK rmx_output_media_get_next_chunk;
-        PFNRMX_OUTPUT_MEDIA_GET_CHUNK_STRIDES rmx_output_media_get_chunk_strides;
-        PFNRMX_OUTPUT_MEDIA_GET_CHUNK_PACKET_SIZES rmx_output_media_get_chunk_packet_sizes;
-        PFNRMX_OUTPUT_MEDIA_SET_CHUNK_OPTION rmx_output_media_set_chunk_option;
-        PFNRMX_OUTPUT_MEDIA_CLEAR_CHUNK_OPTION rmx_output_media_clear_chunk_option;
-        PFNRMX_OUTPUT_MEDIA_CLEAR_CHUNK_ALL_OPTIONS rmx_output_media_clear_chunk_all_options;
         PFNRMX_OUTPUT_MEDIA_MARK_CHUNK_FOR_TRACKING rmx_output_media_mark_chunk_for_tracking;
         PFNRMX_OUTPUT_MEDIA_COMMIT_CHUNK rmx_output_media_commit_chunk;
         PFNRMX_OUTPUT_MEDIA_POLL_FOR_COMPLETION rmx_output_media_poll_for_completion;
@@ -319,9 +269,6 @@ namespace UE::RivermaxCore::Private
         PFNRMX_OUTPUT_GEN_INIT_CHUNK_HANDLE rmx_output_gen_init_chunk_handle;
         PFNRMX_OUTPUT_GEN_GET_NEXT_CHUNK rmx_output_gen_get_next_chunk;
         PFNRMX_OUTPUT_GEN_SET_CHUNK_REMOTE_ADDR rmx_output_gen_set_chunk_remote_addr;
-        PFNRMX_OUTPUT_GEN_SET_CHUNK_OPTION rmx_output_gen_set_chunk_option;
-        PFNRMX_OUTPUT_GEN_CLEAR_CHUNK_OPTION rmx_output_gen_clear_chunk_option;
-        PFNRMX_OUTPUT_GEN_CLEAR_CHUNK_ALL_OPTIONS rmx_output_gen_clear_chunk_all_options;
         PFNRMX_OUTPUT_GEN_APPEND_PACKET_TO_CHUNK rmx_output_gen_append_packet_to_chunk;
         PFNRMX_OUTPUT_GEN_MARK_CHUNK_FOR_TRACKING rmx_output_gen_mark_chunk_for_tracking;
         PFNRMX_OUTPUT_GEN_COMMIT_CHUNK rmx_output_gen_commit_chunk;
@@ -330,8 +277,6 @@ namespace UE::RivermaxCore::Private
         PFNRMX_OUTPUT_GET_CHUNK_COUNT rmx_output_get_chunk_count;
         PFNRMX_OUTPUT_UPDATE_DSCP rmx_output_update_dscp;
         PFNRMX_OUTPUT_UPDATE_ECN rmx_output_update_ecn;
-        PFNRMX_OUTPUT_GET_COMPLETION_USER_TOKEN rmx_output_get_completion_user_token;
-        PFNRMX_OUTPUT_GET_COMPLETION_TIMESTAMP rmx_output_get_completion_timestamp;
     };
 
     bool LoadLibraryFunctions(RIVERMAX_API_FUNCTION_LIST *FunctionList, void *LibraryHandle);
