RiverMaxCore: Fix Call to pure virtual method 'FRivermaxOutStream::CleanupFrameManagement' during destruction has undefined behavior [cplusplus.PureVirtualCall]

#rnx

diff --git a/Engine/Plugins/VirtualProduction/Rivermax/RivermaxCore/Source/RivermaxCore/Private/Streams/RivermaxOutAncStream.cpp b/Engine/Plugins/VirtualProduction/Rivermax/RivermaxCore/Source/RivermaxCore/Private/Streams/RivermaxOutAncStream.cpp
--- a/Engine/Plugins/VirtualProduction/Rivermax/RivermaxCore/Source/RivermaxCore/Private/Streams/RivermaxOutAncStream.cpp
+++ b/Engine/Plugins/VirtualProduction/Rivermax/RivermaxCore/Source/RivermaxCore/Private/Streams/RivermaxOutAncStream.cpp
@@ -170,6 +170,7 @@ namespace UE::RivermaxCore::Private
 
 	void FRivermaxOutAncStream::CleanupFrameManagement()
 	{
+		FRivermaxOutStream::CleanupFrameManagement();
 	}
 
 	bool FRivermaxOutAncStream::ReserveFrame(uint64 FrameCounter) const
diff --git a/Engine/Plugins/VirtualProduction/Rivermax/RivermaxCore/Source/RivermaxCore/Private/Streams/RivermaxOutAncStream.h b/Engine/Plugins/VirtualProduction/Rivermax/RivermaxCore/Source/RivermaxCore/Private/Streams/RivermaxOutAncStream.h
--- a/Engine/Plugins/VirtualProduction/Rivermax/RivermaxCore/Source/RivermaxCore/Private/Streams/RivermaxOutAncStream.h
+++ b/Engine/Plugins/VirtualProduction/Rivermax/RivermaxCore/Source/RivermaxCore/Private/Streams/RivermaxOutAncStream.h
@@ -39,7 +39,7 @@ namespace UE::RivermaxCore::Private
 		virtual bool InitializeStreamMemoryConfig();
 		virtual bool CopyFrameData(const TSharedPtr<FRivermaxOutputFrame>& SourceFrame, uint8* DestinationBase);
 		virtual bool SetupFrameManagement();
-		virtual void CleanupFrameManagement();
+		virtual void CleanupFrameManagement() override;
 		virtual void CompleteCurrentFrame(bool bReleaseFrame) override;
 		virtual void SetupRTPHeaders() override;
 		virtual void LogStreamDescriptionOnCreation() const override;
diff --git a/Engine/Plugins/VirtualProduction/Rivermax/RivermaxCore/Source/RivermaxCore/Private/Streams/RivermaxOutStream.cpp b/Engine/Plugins/VirtualProduction/Rivermax/RivermaxCore/Source/RivermaxCore/Private/Streams/RivermaxOutStream.cpp
--- a/Engine/Plugins/VirtualProduction/Rivermax/RivermaxCore/Source/RivermaxCore/Private/Streams/RivermaxOutStream.cpp
+++ b/Engine/Plugins/VirtualProduction/Rivermax/RivermaxCore/Source/RivermaxCore/Private/Streams/RivermaxOutStream.cpp
@@ -822,6 +822,10 @@ namespace UE::RivermaxCore::Private
 		}
 	}
 
+	void FRivermaxOutStream::CleanupFrameManagement()
+	{
+	}
+
 	void FRivermaxOutStream::DestroyStream_Internal()
 	{
 		DestroyStream();
diff --git a/Engine/Plugins/VirtualProduction/Rivermax/RivermaxCore/Source/RivermaxCore/Private/Streams/RivermaxOutStream.h b/Engine/Plugins/VirtualProduction/Rivermax/RivermaxCore/Source/RivermaxCore/Private/Streams/RivermaxOutStream.h
--- a/Engine/Plugins/VirtualProduction/Rivermax/RivermaxCore/Source/RivermaxCore/Private/Streams/RivermaxOutStream.h
+++ b/Engine/Plugins/VirtualProduction/Rivermax/RivermaxCore/Source/RivermaxCore/Private/Streams/RivermaxOutStream.h
@@ -358,7 +358,7 @@ namespace UE::RivermaxCore::Private
 		virtual bool SetupFrameManagement() = 0;
 
 		/** Clean up frames */
-		virtual void CleanupFrameManagement() = 0;
+		virtual void CleanupFrameManagement();
 
 		/** Destroys rivermax stream. Will wait until it's ready to be destroyed */
 		virtual void DestroyStream();
diff --git a/Engine/Plugins/VirtualProduction/Rivermax/RivermaxCore/Source/RivermaxCore/Private/Streams/RivermaxOutVideoStream.cpp b/Engine/Plugins/VirtualProduction/Rivermax/RivermaxCore/Source/RivermaxCore/Private/Streams/RivermaxOutVideoStream.cpp
--- a/Engine/Plugins/VirtualProduction/Rivermax/RivermaxCore/Source/RivermaxCore/Private/Streams/RivermaxOutVideoStream.cpp
+++ b/Engine/Plugins/VirtualProduction/Rivermax/RivermaxCore/Source/RivermaxCore/Private/Streams/RivermaxOutVideoStream.cpp
@@ -837,6 +837,7 @@ namespace UE::RivermaxCore::Private
 			Allocator->Deallocate();
 			Allocator.Reset();
 		}
+		FRivermaxOutStream::CleanupFrameManagement();
 	}
 
 	SIZE_T FRivermaxOutVideoStream::GetRowSizeInBytes() const
diff --git a/Engine/Plugins/VirtualProduction/Rivermax/RivermaxCore/Source/RivermaxCore/Private/Streams/RivermaxOutVideoStream.h b/Engine/Plugins/VirtualProduction/Rivermax/RivermaxCore/Source/RivermaxCore/Private/Streams/RivermaxOutVideoStream.h
--- a/Engine/Plugins/VirtualProduction/Rivermax/RivermaxCore/Source/RivermaxCore/Private/Streams/RivermaxOutVideoStream.h
+++ b/Engine/Plugins/VirtualProduction/Rivermax/RivermaxCore/Source/RivermaxCore/Private/Streams/RivermaxOutVideoStream.h
@@ -43,7 +43,7 @@ namespace UE::RivermaxCore::Private
 		virtual bool CopyFrameData(const TSharedPtr<FRivermaxOutputFrame>& SourceFrame, uint8* DestinationBase);
 		virtual bool SetupFrameManagement();
 		virtual void SetupRTPHeaders() override;
-		virtual void CleanupFrameManagement();
+		virtual void CleanupFrameManagement() override;
 		virtual void CompleteCurrentFrame(bool bReleaseFrame) override;
 		virtual void LogStreamDescriptionOnCreation() const override;
 		virtual TSharedPtr<FRivermaxOutputFrame> GetNextFrameToSend(bool bWait = false);
