Rivermax: removed asynchronicity from Rivermax Output Streams.

#rnx

diff --git a/Engine/Plugins/VirtualProduction/Rivermax/RivermaxCore/Source/RivermaxCore/Private/Streams/RivermaxOutStream.cpp b/Engine/Plugins/VirtualProduction/Rivermax/RivermaxCore/Source/RivermaxCore/Private/Streams/RivermaxOutStream.cpp
--- a/Engine/Plugins/VirtualProduction/Rivermax/RivermaxCore/Source/RivermaxCore/Private/Streams/RivermaxOutStream.cpp
+++ b/Engine/Plugins/VirtualProduction/Rivermax/RivermaxCore/Source/RivermaxCore/Private/Streams/RivermaxOutStream.cpp
@@ -143,7 +143,10 @@ namespace UE::RivermaxCore::Private
 			return false;
 		}
 
-		InitializationFuture = Async(EAsyncExecution::TaskGraph, [this]()
+		// The following scope used to be initialized asynchronously; however, 
+		// it caused synchronization issues between nodes. When either barrier or 
+		// synchronization correction is added to nDisplay, this can be switched back 
+		// to asynchronous initialization.
 		{
 			// Create Rivermax stream using memory configuration
 			{
@@ -239,23 +242,16 @@ namespace UE::RivermaxCore::Private
 			}
 
 			Listener->OnInitializationCompleted(bIsActive);
-		});
+		}
 
 		return true;
 	}
 
 	void FRivermaxOutStream::Uninitialize()
 	{
-		if (!InitializationFuture.IsReady())
-		{
-			InitializationFuture.Wait();
-		}
-
 		if (RivermaxThread != nullptr)
 		{
 			Stop();
-			
-			CurrentFrame.Reset();
 
 			FrameAvailableSignal->Trigger();
 			FrameReadyToSendSignal->Trigger();
diff --git a/Engine/Plugins/VirtualProduction/Rivermax/RivermaxCore/Source/RivermaxCore/Private/Streams/RivermaxOutStream.h b/Engine/Plugins/VirtualProduction/Rivermax/RivermaxCore/Source/RivermaxCore/Private/Streams/RivermaxOutStream.h
--- a/Engine/Plugins/VirtualProduction/Rivermax/RivermaxCore/Source/RivermaxCore/Private/Streams/RivermaxOutStream.h
+++ b/Engine/Plugins/VirtualProduction/Rivermax/RivermaxCore/Source/RivermaxCore/Private/Streams/RivermaxOutStream.h
@@ -509,9 +509,6 @@ namespace UE::RivermaxCore::Private
 		// TODO: There must be a better place for this.
 		FGuid MonitoringGuid;
 
-		/** Future returned by the async initialization job we launch. Used to detect if it has completed during shutdown. */
-		TFuture<void> InitializationFuture;
-
 		/** Cached cvar values */
 		FOutputStreamCachedCVars CachedCVars;
 
