Rivermax: Fixed monolythincs build issues.

#rnx

diff --git a/Engine/Plugins/VirtualProduction/Rivermax/RivermaxCore/Source/RivermaxCore/Private/RTPHeader.h b/Engine/Plugins/VirtualProduction/Rivermax/RivermaxCore/Source/RivermaxCore/Private/RTPHeader.h
--- a/Engine/Plugins/VirtualProduction/Rivermax/RivermaxCore/Source/RivermaxCore/Private/RTPHeader.h
+++ b/Engine/Plugins/VirtualProduction/Rivermax/RivermaxCore/Source/RivermaxCore/Private/RTPHeader.h
@@ -264,7 +264,11 @@ namespace UE::RivermaxCore::Private
 #if RIVERMAX_PACKET_DEBUG
 		FORCEINLINE void AddField(uint32 InFieldData, int32 InNumBits, const FString& DebugString)
 		{
-			Fields.Add({ InFieldData, InNumBits, DebugString });
+			FFieldInfo Field;
+			Field.Data = InFieldData;
+			Field.NumBits = InNumBits;
+			Field.FieldName = DebugString;
+			Fields.Add(MoveTemp(Field));
 		}
 #else
 		FORCEINLINE void AddField(uint32 InFieldData, int32 InNumBits)
@@ -431,8 +435,11 @@ namespace UE::RivermaxCore::Private
 
 
 #if RIVERMAX_PACKET_DEBUG
-				FFieldInfo Field { BitPosition, CurrentBitsToPack, DebugString };
-				FieldsInOrder[CurrentByte].Add(Field);
+				FFieldInfo Field;
+				Field.BitPosition = BitPosition;
+				Field.NumBits = CurrentBitsToPack;
+				Field.FieldName = DebugString;
+				FieldsInOrder[CurrentByte].Add(MoveTemp(Field));
 #endif
 
 				RemainingBitsInCurrentByte -= CurrentBitsToPack;
diff --git a/Engine/Plugins/VirtualProduction/Rivermax/RivermaxCore/Source/RivermaxCore/Private/RivermaxUtils.cpp b/Engine/Plugins/VirtualProduction/Rivermax/RivermaxCore/Source/RivermaxCore/Private/RivermaxUtils.cpp
--- a/Engine/Plugins/VirtualProduction/Rivermax/RivermaxCore/Source/RivermaxCore/Private/RivermaxUtils.cpp
+++ b/Engine/Plugins/VirtualProduction/Rivermax/RivermaxCore/Source/RivermaxCore/Private/RivermaxUtils.cpp
@@ -182,7 +182,11 @@ namespace UE::RivermaxCore::Private::Utils
 		TArray<TSharedPtr<FJsonValue>> OutputFields;
 
 		const FBigEndianHeaderPacker::FFieldInfo* PreviousBitField = nullptr;
-		FBigEndianHeaderPacker::FFieldInfo Record {0, 0, ""};
+
+		FBigEndianHeaderPacker::FFieldInfo Record;
+		Record.Data = 0;
+		Record.NumBits = 0;
+		Record.FieldName = "";
 
 		const TArray<TArray<FBigEndianHeaderPacker::FFieldInfo>>& FieldsInOrder = HeaderPacker.GetFieldsInOrderByRef();
 		for (const TArray<FBigEndianHeaderPacker::FFieldInfo>& FieldsPerByte : FieldsInOrder)
diff --git a/Engine/Plugins/VirtualProduction/Rivermax/RivermaxCore/Source/RivermaxCore/RivermaxCore.build.cs b/Engine/Plugins/VirtualProduction/Rivermax/RivermaxCore/Source/RivermaxCore/RivermaxCore.build.cs
--- a/Engine/Plugins/VirtualProduction/Rivermax/RivermaxCore/Source/RivermaxCore/RivermaxCore.build.cs
+++ b/Engine/Plugins/VirtualProduction/Rivermax/RivermaxCore/Source/RivermaxCore/RivermaxCore.build.cs
@@ -24,6 +24,7 @@ public class RivermaxCore : ModuleRules
 				"Engine",
 				"Networking",
 				"RenderCore",
+				"Json"
 			}
 		);
 
