Rivermax:
- Fixed the issue with ATC. The code was correctly packing 10-bit ANC words (ST 291-1 framing/parity), but it didn't follow the ATC mapping (SMPTE ST 12-2) for what goes inside those words. ST 12-2 (ATC) requires interleaving: time digit, then user-bits group, repeat.
- Added timed escape in case Rivermax API cannot destroy a stream for some reason.
- Yielding Thread if DPU Queue is full and chunks are unable to be committed.
- fixed minor debugging code issues.

#rnx

diff --git a/Engine/Plugins/VirtualProduction/Rivermax/RivermaxCore/Source/RivermaxCore/Private/RivermaxUtils.cpp b/Engine/Plugins/VirtualProduction/Rivermax/RivermaxCore/Source/RivermaxCore/Private/RivermaxUtils.cpp
--- a/Engine/Plugins/VirtualProduction/Rivermax/RivermaxCore/Source/RivermaxCore/Private/RivermaxUtils.cpp
+++ b/Engine/Plugins/VirtualProduction/Rivermax/RivermaxCore/Source/RivermaxCore/Private/RivermaxUtils.cpp
@@ -69,28 +69,22 @@ namespace UE::RivermaxCore::Private::Utils
 		uint8 BCDMinutes = ToBCD(uint8(InTimecode.Minutes));
 		uint8 BCDHours   = ToBCD(uint8(InTimecode.Hours));
 
-		// Each data word follows the same format: 
+		// Each ATC User Data Word word follows the format: 
 		// 
-		// b0–b2 = 0
+		// b0 - b2 = 0
 		// b3 = DistributedBinaryBits (DBB) flag bit (may carry control info)
-		// b4–b7 = payload nibble
+		// b4 - b7 = payload nibble
 		// b8 = parity
 		// b9 = inverse of b8
-		// DistributedBinaryBits is currently unused. It is used for fractional FPS (Drop-Frame Timecode (DF)) and 
-		// Color-Frame Flag (CF) 
-		//
-		//* UDW1 : Frames units(BCD 0–9)     DBB unused(0)
-		//* UDW2 : Frames tens(BCD 0–2)      DBB = Drop - Frame flag(1 = DF, 0 = NDF)
-		//* UDW3 : Seconds units(BCD 0–9)    DBB unused(0)
-		//* UDW4 : Seconds tens(BCD 0–5)     DBB = Color - Frame flag(1 = CF aligned)
-		//* UDW5 : Minutes units(BCD 0–9)    DBB unused(0)
-		//* UDW6 : Minutes tens(BCD 0–5)     DBB = Binary Group 1 user bit
-		//* UDW7 : Hours units(BCD 0–9)      DBB unused(0)
-		//* UDW8 : Hours tens(BCD 0–2)       DBB = Binary Group 2 user bit
-		//
-		auto MakeUDW10 = [](uint8 Nibble, bool DistributedBinaryBits = false) -> uint16
-		{
-			const uint16 Low8 = uint16(((Nibble & 0xF) << 4) | (DistributedBinaryBits ? 0x8 : 0x0));
+		
+
+		// ExtraMask - to allow DF/CF flags in b6/b7 when needed
+		auto MakeUDW10 = [](uint8 Nibble, bool DistributedBinaryBits = false, uint8 ExtraMask = 0) -> uint16
+		{
+			const uint16 Low8 = uint16(((Nibble & 0xF) << 4)
+				| (DistributedBinaryBits ? 0x8 : 0x0)
+				| (ExtraMask & 0xF0));
+
 			const uint16 ParityEvenBit = (FMath::CountBits(static_cast<uint32>(Low8)) & 1) ? 1u : 0u;
 			const uint16 Bit9 = (~ParityEvenBit) & 0x1;
 			return static_cast<uint16>(Low8 | (ParityEvenBit << 8) | (Bit9 << 9));
@@ -104,21 +98,79 @@ namespace UE::RivermaxCore::Private::Utils
 		const bool bDropFrame = InTimecode.bDropFrameFormat || 
 			((Rate == TwentyNineNineSeven) || (Rate == FiftyNineNintyFour));
 
-		// UDW1–8: BCD digits
-		UserDataWords.Add(MakeUDW10((BCDFrames >> 0) & 0xF)); // Frames units
-		UserDataWords.Add(MakeUDW10((BCDFrames >> 4) & 0xF, bDropFrame)); // Frames tens (DBB can hold Drop-Frame)
-		UserDataWords.Add(MakeUDW10((BCDSeconds >> 0) & 0xF)); // Seconds units
-		UserDataWords.Add(MakeUDW10((BCDSeconds >> 4) & 0xF)); // Seconds tens (DBB can hold Color-Frame)
-		UserDataWords.Add(MakeUDW10((BCDMinutes >> 0) & 0xF)); // Minutes units
-		UserDataWords.Add(MakeUDW10((BCDMinutes >> 4) & 0xF)); // Minutes tens
-		UserDataWords.Add(MakeUDW10((BCDHours >> 0) & 0xF)); // Hours units
-		UserDataWords.Add(MakeUDW10((BCDHours >> 4) & 0xF)); // Hours tens
+		// Not set. Only used for composite analog NTSC/PAL gear.
+		constexpr bool bColorFrame = false;
+
+		{
+			// UDW1: Frames units
+			UserDataWords.Add(MakeUDW10((BCDFrames >> 0) & 0xF, /*DBB*/ false));
+
+			// UDW2: User bits group 1 (unused)
+			UserDataWords.Add(MakeUDW10(0, false));
+		}
+
+		{
+			uint8 NibTensF = (BCDFrames >> 4) & 0xF;
+			uint8 Flags = 0;
+			if (bDropFrame) Flags |= 0x40;          // DF -> b6
+			if (bColorFrame) Flags |= 0x80;         // CF -> b7
+
+			// UDW3: Frames tens + DF (b6) + CF (b7)
+			UserDataWords.Add(MakeUDW10(NibTensF, false, Flags));
+
+			// UDW4: User bits group 2 - Unused.
+			UserDataWords.Add(MakeUDW10(0, false));
+		}
 
-		// UDW9–16: Unused for now.
-		for (int i = 0; i < 8; ++i)
 		{
-			UserDataWords.Add(MakeUDW10(0));
+			// UDW5: Seconds units
+			UserDataWords.Add(MakeUDW10((BCDSeconds >> 0) & 0xF, false));
+
+			// UDW6: User bits group 3 - Unused.
+			UserDataWords.Add(MakeUDW10(0, false));
 		}
+
+		{
+			// UDW7: Seconds tens (optional flag at b7 is unused)
+			UserDataWords.Add(MakeUDW10((BCDSeconds >> 4) & 0xF, false /*DBB*/, 0 /*no flag*/));
+
+			// UDW8: User bits group 4 - unused.
+			UserDataWords.Add(MakeUDW10(0, false));
+		}
+
+		{
+			// UDW9: Minutes units
+			UserDataWords.Add(MakeUDW10((BCDMinutes >> 0) & 0xF, false));
+
+			// UDW10: User bits group 5 - unused
+			UserDataWords.Add(MakeUDW10(0, false));
+		}
+
+		{
+			// UDW11: Minutes tens (+ optional flag at b7 - unused)
+			UserDataWords.Add(MakeUDW10((BCDMinutes >> 4) & 0xF, false, 0));
+
+			// UDW12: User bits group 6 - unused.
+			UserDataWords.Add(MakeUDW10(0, false));
+		}
+
+		{
+			// UDW13: Hours units
+			UserDataWords.Add(MakeUDW10((BCDHours >> 0) & 0xF, false));
+
+			// UDW14: User bits group 7 - unused.
+			UserDataWords.Add(MakeUDW10(0, false));
+		}
+
+		{
+			// UDW15: Hours tens (+ optional flags b6/b7 - unused)
+			UserDataWords.Add(MakeUDW10((BCDHours >> 4) & 0xF, false, 0));
+
+			// UDW16: User bits group 8 - unused.
+			UserDataWords.Add(MakeUDW10(0, false));
+		}
+
+
 		constexpr uint16 MaxNumUDW = 255;
 
 		check(UserDataWords.Num() <= MaxNumUDW);
diff --git a/Engine/Plugins/VirtualProduction/Rivermax/RivermaxCore/Source/RivermaxCore/Private/Streams/RivermaxOutStream.cpp b/Engine/Plugins/VirtualProduction/Rivermax/RivermaxCore/Source/RivermaxCore/Private/Streams/RivermaxOutStream.cpp
--- a/Engine/Plugins/VirtualProduction/Rivermax/RivermaxCore/Source/RivermaxCore/Private/Streams/RivermaxOutStream.cpp
+++ b/Engine/Plugins/VirtualProduction/Rivermax/RivermaxCore/Source/RivermaxCore/Private/Streams/RivermaxOutStream.cpp
@@ -285,6 +285,10 @@ namespace UE::RivermaxCore::Private
 	
 	void FRivermaxOutStream::DestroyStream()
 	{
+		constexpr float TimeOutSecs = 3.f;
+		constexpr float TimeToWaitIncrement = 0.3f;
+		float TimeWaited = 0.f;
+
 		auto CancelUnsentChunks = [&] {
 			rmx_status Status = CachedAPI->rmx_output_media_cancel_unsent_chunks(&StreamData.ChunkHandle);
 			if (Status != RMX_OK)
@@ -304,9 +308,16 @@ namespace UE::RivermaxCore::Private
 				// Sometimes the chunks are not cancelled above and this could get into deadlock unless 
 				// unesnt chunks are cancelled again.
 				CancelUnsentChunks();
-				FPlatformProcess::SleepNoStats(0.3);
+				FPlatformProcess::SleepNoStats(TimeToWaitIncrement);
+				TimeWaited += TimeToWaitIncrement;
 			}
-		} while (Status == RMX_BUSY);
+
+		} while (Status == RMX_BUSY && TimeWaited < TimeOutSecs);
+
+		if (Status == RMX_BUSY)
+		{
+			UE_LOG(LogRivermax, Warning, TEXT("Could not destroy the stream. StreamId: %d"), StreamId);
+		}
 	}
 
 	bool FRivermaxOutStream::WaitForNextRound()
@@ -520,6 +531,9 @@ namespace UE::RivermaxCore::Private
 			else if (Status == RMX_HW_SEND_QUEUE_IS_FULL)
 			{
 				Stats.CommitRetries++;
+
+				// Yeild the current timeslice in case there is another equal priority thread because RMX_HW_SEND_QUEUE_IS_FULL can last a few ms.
+				FPlatformProcess::SleepNoStats(0.f);
 				TRACE_CPUPROFILER_EVENT_SCOPE(CommitNextChunks::QUEUEFULL);
 				++ErrorCount;
 			}
diff --git a/Engine/Plugins/VirtualProduction/Rivermax/RivermaxCore/Source/RivermaxCore/Private/Streams/RivermaxOutVideoStream.cpp b/Engine/Plugins/VirtualProduction/Rivermax/RivermaxCore/Source/RivermaxCore/Private/Streams/RivermaxOutVideoStream.cpp
--- a/Engine/Plugins/VirtualProduction/Rivermax/RivermaxCore/Source/RivermaxCore/Private/Streams/RivermaxOutVideoStream.cpp
+++ b/Engine/Plugins/VirtualProduction/Rivermax/RivermaxCore/Source/RivermaxCore/Private/Streams/RivermaxOutVideoStream.cpp
@@ -241,11 +241,12 @@ namespace UE::RivermaxCore::Private
 		}
 
 
+#if RIVERMAX_PACKET_DEBUG
+		HeaderPacker.Finalize();
+#endif
 
 // Unit test to run when switching platform to make sure that the bit fields are interpreted correctly.
 #if RIVERMAX_USE_BIT_FIELDS & RIVERMAX_PACKET_DEBUG
-		HeaderPacker.Finalize();
-
 		uint8* Ptr = (uint8*) VideoHeader;
 		for (int64 Byte = 0; Byte < FVideoRTPHeader::TwoSRDSize; Byte++)
 		{
@@ -674,7 +675,7 @@ namespace UE::RivermaxCore::Private
 			HeaderMemory = CachedAPI->rmx_output_media_get_dup_sub_block(&Block, StreamMemory.HeaderBlockID);
 			if (HeaderMemory == nullptr)
 			{
-				UE_LOG(LogRivermax, Warning, TEXT("Faild to get header memory block. Output stream won't be created."));
+				UE_LOG(LogRivermax, Warning, TEXT("Failed to get header memory block. Output stream won't be created."));
 				return false;
 			}
 			StreamMemory.RTPHeaders[BlockIndex] = MakeShareable(
@@ -953,9 +954,10 @@ namespace UE::RivermaxCore::Private
 				HeaderRawPtr->RTPHeader.ExtendedSequenceNumber = ByteSwap((uint16)((StreamData.SequenceNumber >> 16) & 0xFFFF));
 #else
 				FBigEndianHeaderPacker HeaderPacker(DataRawPtr, FVideoRTPHeader::TwoSRDSize, false /*bClearExisting*/);
-				HeaderPacker.UpdateField((StreamData.SequenceNumber & 0xFFFF), 16, 16 DEBUG_FIELD_NAME("SEQ"));
-				HeaderPacker.UpdateField(CurrentFrame->MediaTimestamp, 32, 32 DEBUG_FIELD_NAME("Timestamp"));
-				HeaderPacker.UpdateField((StreamData.SequenceNumber >> 16) & 0xFFFF, 16, 96 DEBUG_FIELD_NAME("Extended Sequence Number"));
+				HeaderPacker.Finalize();
+				HeaderPacker.UpdateField((StreamData.SequenceNumber & 0xFFFF), 16, 16 RIVERMAX_DEBUG_FIELD_NAME("SEQ"));
+				HeaderPacker.UpdateField(CurrentFrame->MediaTimestamp, 32, 32 RIVERMAX_DEBUG_FIELD_NAME("Timestamp"));
+				HeaderPacker.UpdateField((StreamData.SequenceNumber >> 16) & 0xFFFF, 16, 96 RIVERMAX_DEBUG_FIELD_NAME("Extended Sequence Number"));
 #endif
 			}
 			else
