Rivermax: Media Output ANC: Beautified the UI and fixed the issue with the IP value not being saved.

#rnx

diff --git a/Engine/Plugins/VirtualProduction/Rivermax/RivermaxCore/Source/RivermaxEditor/Private/Customizations/RivermaxDeviceSelectionCustomization.cpp b/Engine/Plugins/VirtualProduction/Rivermax/RivermaxCore/Source/RivermaxEditor/Private/Customizations/RivermaxDeviceSelectionCustomization.cpp
--- a/Engine/Plugins/VirtualProduction/Rivermax/RivermaxCore/Source/RivermaxEditor/Private/Customizations/RivermaxDeviceSelectionCustomization.cpp
+++ b/Engine/Plugins/VirtualProduction/Rivermax/RivermaxCore/Source/RivermaxEditor/Private/Customizations/RivermaxDeviceSelectionCustomization.cpp
@@ -8,6 +8,7 @@
 #include "IPropertyUtilities.h"
 #include "Widgets/SRivermaxInterfaceComboBox.h"
 #include "PropertyHandle.h"
+#include "ScopedTransaction.h"
 
 namespace UE::RivermaxCore::Utils
 {
@@ -30,7 +31,13 @@ namespace UE::RivermaxCore::Utils
 				.InitialValue(InitialValue)
 				.OnIPAddressSelected_Lambda([InObjectIndex, PropertyHandle](const FString& SelectedIp)
 				{
-					PropertyHandle->SetPerObjectValue(InObjectIndex, SelectedIp);
+						FScopedTransaction Tx(NSLOCTEXT("Rivermax", "ChangeInterfaceAddress", "Change Interface Address"));
+
+						const FPropertyAccess::Result Result = PropertyHandle->SetValue(SelectedIp);
+						if (Result != FPropertyAccess::Success)
+						{
+							UE_LOG(LogTemp, Warning, TEXT("Failed to set InterfaceAddress via IPropertyHandle (Result=%d)"), (int32)Result);
+						}
 				})
 			];
 	}
@@ -41,9 +48,9 @@ namespace UE::RivermaxCore::Utils
 		, const FString& InitialValue
 		, IPropertyTypeCustomizationUtils& CustomizationUtils)
 	{
-		TWeakPtr<IPropertyHandle> WeakInterfaceHandle = InterfaceHandle;
+		IDetailPropertyRow& Row = ChildBuilder.AddProperty(InterfaceHandle);
 
-		ChildBuilder.AddCustomRow(InterfaceHandle->GetPropertyDisplayName())
+		Row.CustomWidget()
 			.NameContent()
 			[
 				InterfaceHandle->CreatePropertyNameWidget()
@@ -53,15 +60,16 @@ namespace UE::RivermaxCore::Utils
 			[
 				SNew(SRivermaxInterfaceComboBox)
 					.InitialValue(InitialValue)
-					.OnIPAddressSelected_Lambda([WeakInterfaceHandle](const FString& SelectedIp)
+					.OnIPAddressSelected_Lambda([InterfaceHandle](const FString& SelectedIp)
+					{
+						FScopedTransaction Tx(NSLOCTEXT("Rivermax", "ChangeInterfaceAddress", "Change Interface Address"));
+
+						const FPropertyAccess::Result Result = InterfaceHandle->SetValue(SelectedIp);
+						if (Result != FPropertyAccess::Success)
 						{
-							if (TSharedPtr<IPropertyHandle> Handle = WeakInterfaceHandle.Pin())
-							{
-								Handle->NotifyPreChange();
-								Handle->SetValue(SelectedIp, EPropertyValueSetFlags::DefaultFlags);
-								Handle->NotifyPostChange(EPropertyChangeType::ValueSet);
-							}
-						})
+							UE_LOG(LogTemp, Warning, TEXT("Failed to set InterfaceAddress via IPropertyHandle (Result=%d)"), (int32)Result);
+						}
+					})
 			];
 	}
 }
diff --git a/Engine/Plugins/VirtualProduction/Rivermax/RivermaxCore/Source/RivermaxEditor/RivermaxEditor.build.cs b/Engine/Plugins/VirtualProduction/Rivermax/RivermaxCore/Source/RivermaxEditor/RivermaxEditor.build.cs
--- a/Engine/Plugins/VirtualProduction/Rivermax/RivermaxCore/Source/RivermaxEditor/RivermaxEditor.build.cs
+++ b/Engine/Plugins/VirtualProduction/Rivermax/RivermaxCore/Source/RivermaxEditor/RivermaxEditor.build.cs
@@ -21,6 +21,7 @@ public class RivermaxEditor : ModuleRules
 				"InputCore",
 				"Slate",
 				"SlateCore",
+				"UnrealEd",
 			}
 		);
 
diff --git a/Engine/Plugins/VirtualProduction/Rivermax/RivermaxMedia/Source/RivermaxMediaEditor/Private/Customization/RivermaxMediaPropertyCustomization.cpp b/Engine/Plugins/VirtualProduction/Rivermax/RivermaxMedia/Source/RivermaxMediaEditor/Private/Customization/RivermaxMediaPropertyCustomization.cpp
--- a/Engine/Plugins/VirtualProduction/Rivermax/RivermaxMedia/Source/RivermaxMediaEditor/Private/Customization/RivermaxMediaPropertyCustomization.cpp
+++ b/Engine/Plugins/VirtualProduction/Rivermax/RivermaxMedia/Source/RivermaxMediaEditor/Private/Customization/RivermaxMediaPropertyCustomization.cpp
@@ -12,6 +12,36 @@
 #include "RivermaxMediaOutput.h"
 
 
+void FRivermaxStreamCustomization::CustomizeHeader(TSharedRef<IPropertyHandle> StructHandle, FDetailWidgetRow& HeaderRow, IPropertyTypeCustomizationUtils& CustomizationUtils)
+{
+	TSharedPtr<IPropertyHandle> StreamTypeHandle = StructHandle->GetChildHandle(GET_MEMBER_NAME_CHECKED(FRivermaxAncStream, StreamType));
+
+	FText HeaderText;
+	if (StreamTypeHandle.IsValid())
+	{
+		FString EnumValue;
+		if (StreamTypeHandle->GetValueAsFormattedString(EnumValue) == FPropertyAccess::Success)
+		{
+			HeaderText = FText::FromString(EnumValue);
+		}
+		else
+		{
+			HeaderText = StreamTypeHandle->GetPropertyDisplayName();
+		}
+	}
+	else
+	{
+		HeaderText = FText::FromString(TEXT("Anc Stream"));
+	}
+
+	HeaderRow.NameContent()
+	[
+		SNew(STextBlock)
+			.Text(HeaderText)
+			.Font(IDetailLayoutBuilder::GetDetailFontBold())
+	];
+}
+
 void FRivermaxStreamCustomization::CustomizeChildren(
 	TSharedRef<IPropertyHandle> StructHandle,
 	IDetailChildrenBuilder& ChildBuilder,
diff --git a/Engine/Plugins/VirtualProduction/Rivermax/RivermaxMedia/Source/RivermaxMediaEditor/Private/Customization/RivermaxMediaPropertyCustomization.h b/Engine/Plugins/VirtualProduction/Rivermax/RivermaxMedia/Source/RivermaxMediaEditor/Private/Customization/RivermaxMediaPropertyCustomization.h
--- a/Engine/Plugins/VirtualProduction/Rivermax/RivermaxMedia/Source/RivermaxMediaEditor/Private/Customization/RivermaxMediaPropertyCustomization.h
+++ b/Engine/Plugins/VirtualProduction/Rivermax/RivermaxMedia/Source/RivermaxMediaEditor/Private/Customization/RivermaxMediaPropertyCustomization.h
@@ -13,7 +13,7 @@ public:
 	}
 
 	// IPropertyTypeCustomization
-	virtual void CustomizeHeader(TSharedRef<IPropertyHandle> StructHandle, FDetailWidgetRow& HeaderRow, IPropertyTypeCustomizationUtils& CustomizationUtils) override {};
+	virtual void CustomizeHeader(TSharedRef<IPropertyHandle> StructHandle, FDetailWidgetRow& HeaderRow, IPropertyTypeCustomizationUtils& CustomizationUtils) override;
 
 	virtual void CustomizeChildren(TSharedRef<IPropertyHandle> StructHandle, IDetailChildrenBuilder& ChildBuilder,
 		IPropertyTypeCustomizationUtils& CustomizationUtils) override;
