Rivermax: Fix debug build failures due to missing includes.

diff --git a/Engine/Plugins/VirtualProduction/Rivermax/RivermaxCore/Source/RivermaxCore/Private/RivermaxUtils.cpp b/Engine/Plugins/VirtualProduction/Rivermax/RivermaxCore/Source/RivermaxCore/Private/RivermaxUtils.cpp
--- a/Engine/Plugins/VirtualProduction/Rivermax/RivermaxCore/Source/RivermaxCore/Private/RivermaxUtils.cpp
+++ b/Engine/Plugins/VirtualProduction/Rivermax/RivermaxCore/Source/RivermaxCore/Private/RivermaxUtils.cpp
@@ -7,6 +7,11 @@
 #include "RivermaxTypes.h"
 #include "RTPHeader.h"
 
+#if RIVERMAX_PACKET_DEBUG
+#include "Serialization/JsonSerializer.h"
+#include "Misc/FileHelper.h"
+#endif
+
 namespace UE::RivermaxCore::Private::Utils
 {
 	uint32 TimestampToFrameNumber(uint32 Timestamp, const FFrameRate& FrameRate)
