Rivermax:
- fixed an issue with 16 bit RGB stream getting cropped on input due to invalid number of dispatch group count exceeding maximum allowed.
- Cleanup and minor improvements

[RN] Fixed an issue with large inners getting cropped due to a large number of compute shader groups.

diff --git a/Engine/Plugins/VirtualProduction/Rivermax/RivermaxCore/Shaders/Private/RivermaxShaders.usf b/Engine/Plugins/VirtualProduction/Rivermax/RivermaxCore/Shaders/Private/RivermaxShaders.usf
--- a/Engine/Plugins/VirtualProduction/Rivermax/RivermaxCore/Shaders/Private/RivermaxShaders.usf
+++ b/Engine/Plugins/VirtualProduction/Rivermax/RivermaxCore/Shaders/Private/RivermaxShaders.usf
@@ -21,6 +21,9 @@ uint OutputTextureSizeX;
 uint TotalElementCount;
 RWTexture2D<float4> OutTexture;
 
+/** This value needs to be reflected in GetComputeShaderGroupCount in RivermaxMediaUtils */
+#define DISPATCH_GROUP_COUNT 512, 1, 1
+
 struct FYUV8Bit422Buffer
 {
 	uint DWord0;
@@ -73,7 +76,7 @@ StructuredBuffer<FRGB16fBuffer> InputRGB16fBuffer;
 
 // Shader to convert from RGB to YUV 422 8 bits
 
-[numthreads(64,1,1)]
+[numthreads(DISPATCH_GROUP_COUNT)]
 void RGBToYUV8Bit422(uint3 ThreadId : SV_DispatchThreadID, uint GroupIndex : SV_GroupIndex)
 {
 	// Make sure it's a valid thread
@@ -105,7 +108,7 @@ void RGBToYUV8Bit422(uint3 ThreadId : SV_DispatchThreadID, uint GroupIndex : SV_
 
 // Shader to convert from RGB to YUV 422 10 bits packed little endian
 
-[numthreads(64,1,1)]
+[numthreads(DISPATCH_GROUP_COUNT)]
 void RGBToYUV10Bit422(uint3 ThreadId : SV_DispatchThreadID, uint GroupIndex : SV_GroupIndex)
 {
 	// Make sure it's a valid thread
@@ -141,7 +144,7 @@ void RGBToYUV10Bit422(uint3 ThreadId : SV_DispatchThreadID, uint GroupIndex : SV
 
 // Shader to convert from RGBA to RGB10
 
-[numthreads(64,1,1)]
+[numthreads(DISPATCH_GROUP_COUNT)]
 void RGBToRGB10BitCS(uint3 ThreadId : SV_DispatchThreadID, uint GroupIndex : SV_GroupIndex)
 {
 	// Make sure it's a valid thread
@@ -198,7 +201,7 @@ void RGBToRGB10BitCS(uint3 ThreadId : SV_DispatchThreadID, uint GroupIndex : SV_
 
 // Shader to convert from packed RGB10 to RGB10A2 texture
 
-[numthreads(64,1,1)]
+[numthreads(DISPATCH_GROUP_COUNT)]
 void RGB10BitToRGBACS(uint3 ThreadId : SV_DispatchThreadID, uint GroupIndex : SV_GroupIndex)
 {
 	// Make sure it's a valid thread
@@ -302,7 +305,7 @@ void RGB10BitToRGBACS(uint3 ThreadId : SV_DispatchThreadID, uint GroupIndex : SV
 
 // Shader to convert from RGBA to RGB8
 
-[numthreads(64,1,1)]
+[numthreads(DISPATCH_GROUP_COUNT)]
 void RGBToRGB8BitCS(uint3 ThreadId : SV_DispatchThreadID, uint GroupIndex : SV_GroupIndex)
 {
 	// Make sure it's a valid thread
@@ -336,7 +339,7 @@ void RGBToRGB8BitCS(uint3 ThreadId : SV_DispatchThreadID, uint GroupIndex : SV_G
 
 // Shader to convert from packed RGB8 to RGBA8 texture
 
-[numthreads(64,1,1)]
+[numthreads(DISPATCH_GROUP_COUNT)]
 void RGB8BitToRGBA8CS(uint3 ThreadId : SV_DispatchThreadID, uint GroupIndex : SV_GroupIndex)
 {
 	// Make sure it's a valid thread
@@ -393,7 +396,7 @@ void RGB8BitToRGBA8CS(uint3 ThreadId : SV_DispatchThreadID, uint GroupIndex : SV
 
 // Shader to convert from RGBA to RGB12
 
-[numthreads(64,1,1)]
+[numthreads(DISPATCH_GROUP_COUNT)]
 void RGBToRGB12BitCS(uint3 ThreadId : SV_DispatchThreadID, uint GroupIndex : SV_GroupIndex)
 {
 	// Make sure it's a valid thread
@@ -432,7 +435,7 @@ void RGBToRGB12BitCS(uint3 ThreadId : SV_DispatchThreadID, uint GroupIndex : SV_
 
 // Shader to convert from packed RGB12 to float RGBA texture
 
-[numthreads(64,1,1)]
+[numthreads(DISPATCH_GROUP_COUNT)]
 void RGB12BitToRGBACS(uint3 ThreadId : SV_DispatchThreadID, uint GroupIndex : SV_GroupIndex)
 {
 	// Make sure it's a valid thread
@@ -503,7 +506,7 @@ void RGB12BitToRGBACS(uint3 ThreadId : SV_DispatchThreadID, uint GroupIndex : SV
 }
 
 // Shader to convert from RGBA to RGB Float16
-[numthreads(64,1,1)]
+[numthreads(DISPATCH_GROUP_COUNT)]
 void RGBToRGB16fBitCS(uint3 ThreadId : SV_DispatchThreadID, uint GroupIndex : SV_GroupIndex)
 {
 	// Make sure it's a valid thread
@@ -566,7 +569,7 @@ float GetEncodedFloatValue(in uint Value)
 
 // Shader to convert from packed RGB16f to float RGBA texture
 
-[numthreads(64,1,1)]
+[numthreads(DISPATCH_GROUP_COUNT)]
 void RGB16fBitToRGBACS(uint3 ThreadId : SV_DispatchThreadID, uint GroupIndex : SV_GroupIndex)
 {
 	// Make sure it's a valid thread
@@ -621,7 +624,7 @@ void RGB16fBitToRGBACS(uint3 ThreadId : SV_DispatchThreadID, uint GroupIndex : S
 
 // Shader to convert from YUV 422 8 bits to RGBA
 
-[numthreads(64,1,1)]
+[numthreads(DISPATCH_GROUP_COUNT)]
 void YUV8Bit422ToRGBACS(uint3 ThreadId : SV_DispatchThreadID, uint GroupIndex : SV_GroupIndex)
 {
 	// Make sure it's a valid thread
@@ -678,7 +681,7 @@ void YUV8Bit422ToRGBACS(uint3 ThreadId : SV_DispatchThreadID, uint GroupIndex :
 
 // Shader to convert from YUV 422 10 bits packed little endian to RGBA
 
-[numthreads(64,1,1)]
+[numthreads(DISPATCH_GROUP_COUNT)]
 void YUV10Bit422ToRGBACS(uint3 ThreadId : SV_DispatchThreadID, uint GroupIndex : SV_GroupIndex)
 {
 	// Make sure it's a valid thread
diff --git a/Engine/Plugins/VirtualProduction/Rivermax/RivermaxMedia/Source/RivermaxMedia/Private/RivermaxMediaCapture.cpp b/Engine/Plugins/VirtualProduction/Rivermax/RivermaxMedia/Source/RivermaxMedia/Private/RivermaxMediaCapture.cpp
--- a/Engine/Plugins/VirtualProduction/Rivermax/RivermaxMedia/Source/RivermaxMedia/Private/RivermaxMediaCapture.cpp
+++ b/Engine/Plugins/VirtualProduction/Rivermax/RivermaxMedia/Source/RivermaxMedia/Private/RivermaxMediaCapture.cpp
@@ -514,7 +514,7 @@ void URivermaxMediaCapture::OnCustomCapture_RenderingThread(FRDGBuilder& GraphBu
 	const FIntRect ViewRect(CopyInfo.GetSourceRect());
 	constexpr bool bDoLinearToSRGB = false;
 	const FIntPoint AlignedBufferSize = StreamOptions->AlignedResolution;
-	const FIntVector GroupCount = FComputeShaderUtils::GetGroupCount(DesiredOutputSize.X, 64);
+	const FIntVector GroupCount = UE::RivermaxMediaUtils::Private::GetComputeShaderGroupCount(DesiredOutputSize.X);
 
 	switch (RivermaxOutput->VideoStream.PixelFormat)
 	{
diff --git a/Engine/Plugins/VirtualProduction/Rivermax/RivermaxMedia/Source/RivermaxMedia/Private/RivermaxMediaPlayer.cpp b/Engine/Plugins/VirtualProduction/Rivermax/RivermaxMedia/Source/RivermaxMedia/Private/RivermaxMediaPlayer.cpp
--- a/Engine/Plugins/VirtualProduction/Rivermax/RivermaxMedia/Source/RivermaxMedia/Private/RivermaxMediaPlayer.cpp
+++ b/Engine/Plugins/VirtualProduction/Rivermax/RivermaxMedia/Source/RivermaxMedia/Private/RivermaxMediaPlayer.cpp
@@ -655,8 +655,7 @@ namespace UE::RivermaxMedia
 			}
 
 			const FIntPoint ProcessedOutputDimension = { (int32)SourceBufferDesc.NumberOfElements, 1 };
-			const FIntVector GroupCountUnclamped = FComputeShaderUtils::GetGroupCount(SourceBufferDesc.NumberOfElements, 64);
-			const FIntVector GroupCount = GroupCountUnclamped.ComponentMin(GRHIGlobals.MaxDispatchThreadGroupsPerDimension);
+			const FIntVector GroupCount = UE::RivermaxMediaUtils::Private::GetComputeShaderGroupCount(SourceBufferDesc.NumberOfElements);
 			FGlobalShaderMap* GlobalShaderMap = GetGlobalShaderMap(GMaxRHIFeatureLevel);
 
 			//Configure shader and add conversion pass based on desired pixel format
diff --git a/Engine/Plugins/VirtualProduction/Rivermax/RivermaxMedia/Source/RivermaxMedia/Private/RivermaxMediaUtils.cpp b/Engine/Plugins/VirtualProduction/Rivermax/RivermaxMedia/Source/RivermaxMedia/Private/RivermaxMediaUtils.cpp
--- a/Engine/Plugins/VirtualProduction/Rivermax/RivermaxMedia/Source/RivermaxMedia/Private/RivermaxMediaUtils.cpp
+++ b/Engine/Plugins/VirtualProduction/Rivermax/RivermaxMedia/Source/RivermaxMedia/Private/RivermaxMediaUtils.cpp
@@ -365,6 +365,25 @@ namespace UE::RivermaxMediaUtils::Private
 		}
 	}
 
+	FIntVector GetComputeShaderGroupCount(uint32 NumberOfElements)
+	{
+		constexpr uint32 GX = 512u;
+		constexpr uint32 GY = 1u;
+		constexpr uint32 ThreadsPerGroup = GX * GY;
+
+		const uint32 MaxGroupsPerDim = GRHIMaxDispatchThreadGroupsPerDimension.X; // typically 65535
+
+		const uint64 GroupsNeeded = FMath::DivideAndRoundUp((uint64)NumberOfElements, (uint64)ThreadsPerGroup);
+
+		const uint32 GroupsX = (uint32)FMath::Min<uint64>(GroupsNeeded, MaxGroupsPerDim);
+
+		//TODO:
+		// GPU compute dispatches have maximum workgroup counts per dimension. 
+		// Unreal exposes these limits via GRHIMaxDispatchThreadGroupsPerDimension.X/Y/Z. The current implementation clamps X 
+		// meaning that 16 bit 8k*8k resolution gets clipped. In the future release we should add multi-dimensional dispatch support.
+		return FIntVector(GroupsX, 1, 1);
+	}
+
 	bool StreamOptionsToSDPDescription(const UE::RivermaxCore::FRivermaxOutputOptions& Options, TArray<char>& OutSDPDescription)
 	{
 		uint8 NumberOfStreams = 0; 
diff --git a/Engine/Plugins/VirtualProduction/Rivermax/RivermaxMedia/Source/RivermaxMedia/Private/RivermaxMediaUtils.h b/Engine/Plugins/VirtualProduction/Rivermax/RivermaxMedia/Source/RivermaxMedia/Private/RivermaxMediaUtils.h
--- a/Engine/Plugins/VirtualProduction/Rivermax/RivermaxMedia/Source/RivermaxMedia/Private/RivermaxMediaUtils.h
+++ b/Engine/Plugins/VirtualProduction/Rivermax/RivermaxMedia/Source/RivermaxMedia/Private/RivermaxMediaUtils.h
@@ -49,6 +49,19 @@ namespace UE::RivermaxMediaUtils::Private
 	UE::RivermaxCore::EFrameLockingMode MediaOutputFrameLockingToRivermax(ERivermaxFrameLockingMode InFrameLockingMode);
 	FSourceBufferDesc GetBufferDescription(const FIntPoint& Resolution, ERivermaxMediaSourcePixelFormat InPixelFormat);
 
+	/**
+	 * Computes the dispatch group count for a 2110 buffer conversion compute shader.
+	 *
+	 * Each group in the shader is declared as [numthreads(512,1,1)], which means
+	 * 512 threads are launched per group. In the future we should consider using multi dimensional 
+	 * dispatch group count with a smaller group size.
+	 * NumberOfElements - The total number of elements (e.g. pixels or items) that need to be processed.
+	 * 
+	 * Returns FIntVector giving the number of groups to dispatch in each dimension. 
+	 * It needs to reflect DISPATCH_GROUP_COUNT in RivermaxShaders.usf.
+	 */
+	FIntVector GetComputeShaderGroupCount(uint32 NumberOfElements);
+
 	/** Gets a resolution aligned to a pixel group specified for the provided video format as per 2110-20 requirements. */
 	FIntPoint GetAlignedResolution(const UE::RivermaxCore::FVideoFormatInfo& InFormatInfo, const FIntPoint& ResolutionToAlign);
 
diff --git a/Engine/Plugins/VirtualProduction/Rivermax/RivermaxMedia/Source/RivermaxMedia/Public/RivermaxMediaOutput.h b/Engine/Plugins/VirtualProduction/Rivermax/RivermaxMedia/Source/RivermaxMedia/Public/RivermaxMediaOutput.h
--- a/Engine/Plugins/VirtualProduction/Rivermax/RivermaxMedia/Source/RivermaxMedia/Public/RivermaxMediaOutput.h
+++ b/Engine/Plugins/VirtualProduction/Rivermax/RivermaxMedia/Source/RivermaxMedia/Public/RivermaxMediaOutput.h
@@ -236,9 +236,11 @@ public:
 	UPROPERTY()
 	bool bUseGPUDirect_DEPRECATED = true;
 
+	/** Video output stream properties. */
 	UPROPERTY(EditAnywhere, BlueprintReadWrite, Category = "Video")
 	FRivermaxVideoStream VideoStream;
 
+	/** Collection of Ancillary streams. */
 	UPROPERTY(EditAnywhere, BlueprintReadWrite, Category = "Anc")
 	TArray<FRivermaxAncStream> AncStreams;
 
diff --git a/Engine/Plugins/VirtualProduction/Rivermax/RivermaxMedia/Source/RivermaxMediaEditor/Private/Customization/RivermaxMediaPropertyCustomization.cpp b/Engine/Plugins/VirtualProduction/Rivermax/RivermaxMedia/Source/RivermaxMediaEditor/Private/Customization/RivermaxMediaPropertyCustomization.cpp
--- a/Engine/Plugins/VirtualProduction/Rivermax/RivermaxMedia/Source/RivermaxMediaEditor/Private/Customization/RivermaxMediaPropertyCustomization.cpp
+++ b/Engine/Plugins/VirtualProduction/Rivermax/RivermaxMedia/Source/RivermaxMediaEditor/Private/Customization/RivermaxMediaPropertyCustomization.cpp
@@ -31,7 +31,21 @@ void FRivermaxStreamCustomization::CustomizeHeader(TSharedRef<IPropertyHandle> S
 	}
 	else
 	{
-		HeaderText = FText::FromString(TEXT("Anc Stream"));
+		if (FProperty* Prop = StructHandle->GetProperty())
+		{
+			if (const FStructProperty* StructProp = CastField<FStructProperty>(Prop))
+			{
+				const UScriptStruct* SS = StructProp->Struct;
+				if (SS == FRivermaxAncStream::StaticStruct())
+				{
+					HeaderText = FText::FromString(TEXT("ANC Stream"));
+				}
+				else if (SS == FRivermaxVideoStream::StaticStruct())
+				{
+					HeaderText = FText::FromString(TEXT("Video Stream"));
+				}
+			}
+		}
 	}
 
 	HeaderRow.NameContent()
diff --git a/Engine/Plugins/VirtualProduction/Rivermax/RivermaxMedia/Source/RivermaxMediaEditor/Private/RivermaxMediaEditorModule.cpp b/Engine/Plugins/VirtualProduction/Rivermax/RivermaxMedia/Source/RivermaxMediaEditor/Private/RivermaxMediaEditorModule.cpp
--- a/Engine/Plugins/VirtualProduction/Rivermax/RivermaxMedia/Source/RivermaxMediaEditor/Private/RivermaxMediaEditorModule.cpp
+++ b/Engine/Plugins/VirtualProduction/Rivermax/RivermaxMedia/Source/RivermaxMediaEditor/Private/RivermaxMediaEditorModule.cpp
@@ -41,7 +41,6 @@ private:
 		};
 
 		AddCustomization(URivermaxMediaSource::StaticClass()->GetFName(), FOnGetDetailCustomizationInstance::CreateStatic(&FRivermaxMediaDetailsCustomization::MakeInstance));
-		AddCustomization(URivermaxMediaOutput::StaticClass()->GetFName(), FOnGetDetailCustomizationInstance::CreateStatic(&FRivermaxMediaDetailsCustomization::MakeInstance));
 
 		PropertyModule.RegisterCustomPropertyTypeLayout(TEXT("RivermaxStream"), FOnGetPropertyTypeCustomizationInstance::CreateStatic(&FRivermaxStreamCustomization::MakeInstance));
 		PropertyModule.RegisterCustomPropertyTypeLayout(TEXT("RivermaxVideoStream"), FOnGetPropertyTypeCustomizationInstance::CreateStatic(&FRivermaxStreamCustomization::MakeInstance));
