Rivermax: Exposed CPU core affinity setting via a CVAR for Rivermax internal thread. Pinning to a specific core can reduce scheduling jitter.

#rn Added ability to set core affinity to Rivermax internal thread

#Synced-CL 0

diff --git a/Engine/Plugins/VirtualProduction/Rivermax/RivermaxCore/Source/RivermaxCore/Private/RivermaxManager.cpp b/Engine/Plugins/VirtualProduction/Rivermax/RivermaxCore/Source/RivermaxCore/Private/RivermaxManager.cpp
--- a/Engine/Plugins/VirtualProduction/Rivermax/RivermaxCore/Source/RivermaxCore/Private/RivermaxManager.cpp
+++ b/Engine/Plugins/VirtualProduction/Rivermax/RivermaxCore/Source/RivermaxCore/Private/RivermaxManager.cpp
@@ -63,6 +63,11 @@ namespace UE::RivermaxCore::Private
 		TEXT("Whether GPUDirect capability is enabled for output purposes."),
 		ECVF_Default);
 
+	static TAutoConsoleVariable<int32> CVarRivermaxCoreAffinity(
+		TEXT("Rivermax.CPUAffinity"), -1,
+		TEXT("CPU Core to pin Rivermax output/input thread to. Negative means disabled. Core index is 0..X"),
+		ECVF_Default);
+
 
 	static uint64 RivermaxTimeHandler(void* Context)
 	{
@@ -164,7 +169,7 @@ namespace UE::RivermaxCore::Private
 		}
 		else
 		{
-			UE_LOG(LogRivermax, Display, TEXT("Succesfully loaded Rivermax library from '%s'"), *LibraryPath);
+			UE_LOG(LogRivermax, Display, TEXT("Successfully loaded Rivermax library from '%s'"), *LibraryPath);
 		}
 
 		return bIsLibraryValid;
@@ -223,6 +228,51 @@ namespace UE::RivermaxCore::Private
 		return TimeSource;
 	}
 
+	void FRivermaxManager::SetCPUAffinity(int32 CpuCore)
+	{
+		if (CpuCore < 0)
+		{
+			UE_LOG(LogRivermax, Log, TEXT("Rivermax CPU affinity not set (CpuCore=%d)"), CpuCore);
+			return;
+		}
+
+		const int32 NumLogicalCores = FPlatformMisc::NumberOfCoresIncludingHyperthreads();
+		if (NumLogicalCores <= 0)
+		{
+			UE_LOG(LogRivermax, Warning, TEXT("Rivermax CPU affinity not set: could not query core count (CpuCore=%d)"), CpuCore);
+			return;
+		}
+
+		if (CpuCore >= NumLogicalCores)
+		{
+			UE_LOG(LogRivermax, Warning, TEXT("Rivermax CPU affinity not set: CpuCore=%d out of range [0..%d]"), CpuCore, NumLogicalCores - 1);
+			return;
+		}
+
+		constexpr int32 BitsPerWord = 64;
+		const int32 NumBits = CpuCore + 1;
+		const int32 NumWords = (NumBits + BitsPerWord - 1) / BitsPerWord;
+
+		TArray<uint64> CpuMask;
+		CpuMask.Init(0ull, NumWords);
+
+		const int32 WordIndex = CpuCore / BitsPerWord;
+		const int32 BitOffset = CpuCore % BitsPerWord;
+		CpuMask[WordIndex] |= (1ULL << BitOffset);
+
+		const rmx_status Status = GetApi()->rmx_set_cpu_affinity(CpuMask.GetData(), static_cast<size_t>(NumBits));
+
+		if (Status != RMX_OK)
+		{
+			UE_LOG(LogRivermax, Error, TEXT("Failed to set affinity to core %d (status=%d)"), CpuCore, (int32)Status);
+		}
+		else
+		{
+			UE_LOG(LogRivermax, Log, TEXT("Set CPU affinity to core %d"), CpuCore);
+		}
+
+	}
+
 	void FRivermaxManager::InitializeLibrary()
 	{
 		bool bCanProceed = FApp::CanEverRender();
@@ -254,7 +304,10 @@ namespace UE::RivermaxCore::Private
 
 			// Bug with 1.31 version where this is required to be called before _init() even though it is optional
 			GetApi()->rmx_enable_system_signal_handling();
-			
+
+			const int32 Core = CVarRivermaxCoreAffinity.GetValueOnGameThread();
+			SetCPUAffinity(Core);
+
 			const rmx_status Status = GetApi()->_rmx_init(&Policy);
 			if (Status == RMX_OK)
 			{
@@ -296,7 +349,7 @@ namespace UE::RivermaxCore::Private
 						}
 						else
 						{
-							UE_LOG(LogRivermax, Log, TEXT("Rivermax library version %d.%d.%d succesfully initialized, %s"), CurrentVersion->major, CurrentVersion->minor, CurrentVersion->patch, GPUDirectSupport);
+							UE_LOG(LogRivermax, Log, TEXT("Rivermax library version %d.%d.%d successfully initialized, %s"), CurrentVersion->major, CurrentVersion->minor, CurrentVersion->patch, GPUDirectSupport);
 						}
 					}
 					else
diff --git a/Engine/Plugins/VirtualProduction/Rivermax/RivermaxCore/Source/RivermaxCore/Private/RivermaxManager.h b/Engine/Plugins/VirtualProduction/Rivermax/RivermaxCore/Source/RivermaxCore/Private/RivermaxManager.h
--- a/Engine/Plugins/VirtualProduction/Rivermax/RivermaxCore/Source/RivermaxCore/Private/RivermaxManager.h
+++ b/Engine/Plugins/VirtualProduction/Rivermax/RivermaxCore/Source/RivermaxCore/Private/RivermaxManager.h
@@ -45,6 +45,9 @@ namespace UE::RivermaxCore::Private
 		/** Initializes rivermax library */
 		void InitializeLibrary();
 
+		/** Sets CPU affinity for Rivermax*/
+		void SetCPUAffinity(int32 CpuCore);
+
 		/** Verify if system is capable of doing GPUDirect */
 		void VerifyGPUDirectCapability();
 
