Rivermax:
- Fixed GPUDirect Cuda context switching.
- Fixed incorrect allocation of weak pointer for payload on a stack that was destroyed upon scope exit.
- Added GPUDirect stream synchronization on deinitization.
- fixed Cuda External memory handle leak

#rn Rivermax: Better thread safety and resource handling with GPUDirect on input.

#Synced-CL 0

diff --git a/Engine/Plugins/VirtualProduction/Rivermax/RivermaxCore/Source/RivermaxCore/Private/RivermaxManager.cpp b/Engine/Plugins/VirtualProduction/Rivermax/RivermaxCore/Source/RivermaxCore/Private/RivermaxManager.cpp
--- a/Engine/Plugins/VirtualProduction/Rivermax/RivermaxCore/Source/RivermaxCore/Private/RivermaxManager.cpp
+++ b/Engine/Plugins/VirtualProduction/Rivermax/RivermaxCore/Source/RivermaxCore/Private/RivermaxManager.cpp
@@ -65,7 +65,8 @@ namespace UE::RivermaxCore::Private
 
 	static TAutoConsoleVariable<int32> CVarRivermaxCoreAffinity(
 		TEXT("Rivermax.CPUAffinity"), -1,
-		TEXT("CPU Core to pin Rivermax output/input thread to. Negative means disabled. Core index is 0..X"),
+		TEXT("CPU Core to pin Rivermax output/input thread to. Negative means disabled. Core index is 0..X.\n\
+			The best choice for the cpu core is one that is on the same NUMA node as the DPU/NIC used for Rivermax traffic."),
 		ECVF_Default);
 
 
diff --git a/Engine/Plugins/VirtualProduction/Rivermax/RivermaxCore/Source/RivermaxCore/Private/Streams/RivermaxInputStream.cpp b/Engine/Plugins/VirtualProduction/Rivermax/RivermaxCore/Source/RivermaxCore/Private/Streams/RivermaxInputStream.cpp
--- a/Engine/Plugins/VirtualProduction/Rivermax/RivermaxCore/Source/RivermaxCore/Private/Streams/RivermaxInputStream.cpp
+++ b/Engine/Plugins/VirtualProduction/Rivermax/RivermaxCore/Source/RivermaxCore/Private/Streams/RivermaxInputStream.cpp
@@ -22,6 +22,50 @@
 
 namespace UE::RivermaxCore::Private
 {
+	struct FRivermaxScopedCudaContext
+	{
+		explicit FRivermaxScopedCudaContext(CUcontext InCtx) : Ctx(InCtx)
+		{
+			if (Ctx)
+			{
+				const CUresult Res = FCUDAModule::CUDA().cuCtxPushCurrent(Ctx);
+				bPushed = (Res == CUDA_SUCCESS);
+#if DO_CHECK
+				if (!bPushed)
+				{
+					UE_LOG(LogRivermax, Warning, TEXT("cuCtxPushCurrent failed (err=%d)"), (int32)Res);
+				}
+#endif
+			}
+		}
+
+		~FRivermaxScopedCudaContext()
+		{
+			if (bPushed)
+			{
+				CUcontext Popped = nullptr;
+				const CUresult Res = FCUDAModule::CUDA().cuCtxPopCurrent(&Popped);
+#if DO_CHECK
+				if (Res != CUDA_SUCCESS)
+				{
+					UE_LOG(LogRivermax, Warning, TEXT("cuCtxPopCurrent failed (err=%d)"), (int32)Res);
+				}
+				else if (Popped != Ctx)
+				{
+					UE_LOG(LogRivermax, Warning, TEXT("cuCtxPopCurrent popped unexpected context"));
+				}
+#endif
+			}
+		}
+
+	private:
+		// The CUDA context to push/pop.
+		CUcontext Ctx = nullptr;
+
+		// since cuCtxPopCurrent is a no-op if cuCtxPushCurrent failed, track whether we pushed successfully.
+		bool bPushed = false;
+	};
+
 	/** Initialization calls to Rivermax SDK isn't threadsafe and to prevent stalling game thread, a critical section is shared across streams */
 	static FCriticalSection InitCriticalSection;
 
@@ -166,7 +210,7 @@ namespace UE::RivermaxCore::Private
 			FScopeLock Lock(&InitCriticalSection);
 
 			// If the stream is trying to shutdown before the init task has even started, don't bother
-			if (bIsShuttingDown)
+			if (bIsShuttingDown.load(std::memory_order_acquire))
 			{
 				return;
 			}
@@ -223,6 +267,18 @@ namespace UE::RivermaxCore::Private
 			InitTaskFuture.Wait();
 		}
 
+		// Since cuda callback was fed raw payload pointer, GPU Stream needs to be synchronized to make sure the callback is called and there are no memory leaks.
+		if (bIsUsingGPUDirect && GPUStream)
+		{
+			FCUDAModule& CudaModule = FModuleManager::GetModuleChecked<FCUDAModule>("CUDA");
+			FRivermaxScopedCudaContext CudaScopedContext(CudaModule.GetCudaContextForDevice(GPUDeviceIndex));
+			CUresult Status = CudaModule.DriverAPI()->cuStreamSynchronize(reinterpret_cast<CUstream>(GPUStream));
+			if (Status != CUDA_SUCCESS)
+			{
+				UE_LOG(LogRivermax, Warning, TEXT("CUDA Failed to synchronize stream. Status: %d"), Status);
+			}
+		}
+
 		if (RivermaxThread != nullptr)
 		{
 			Stop();
@@ -574,8 +630,7 @@ namespace UE::RivermaxCore::Private
 
 		FCUDAModule& CudaModule = FModuleManager::GetModuleChecked<FCUDAModule>("CUDA");
 		GPUDeviceIndex = CudaModule.GetCudaDeviceIndex();
-		CudaModule.DriverAPI()->cuCtxPushCurrent(CudaModule.GetCudaContextForDevice(GPUDeviceIndex));
-
+		FRivermaxScopedCudaContext CudaScopedContext(CudaModule.GetCudaContextForDevice(GPUDeviceIndex));
 		// Todo: Add support for mgpu. 
 		CUdevice CudaDevice;
 		CUresult Status = CudaModule.DriverAPI()->cuDeviceGet(&CudaDevice, GPUDeviceIndex);
@@ -692,7 +747,7 @@ namespace UE::RivermaxCore::Private
 		if (GPUAllocatedMemorySize > 0)
 		{
 			FCUDAModule& CudaModule = FModuleManager::GetModuleChecked<FCUDAModule>("CUDA");
-			CudaModule.DriverAPI()->cuCtxPushCurrent(CudaModule.GetCudaContextForDevice(GPUDeviceIndex));
+			FRivermaxScopedCudaContext CudaScopedContext(CudaModule.GetCudaContextForDevice(GPUDeviceIndex));
 
 			const CUdeviceptr BaseAddress = reinterpret_cast<CUdeviceptr>(GPUAllocatedMemoryBaseAddress);
 			CUresult Status = CudaModule.DriverAPI()->cuMemUnmap(BaseAddress, GPUAllocatedMemorySize);
@@ -712,11 +767,16 @@ namespace UE::RivermaxCore::Private
 			GPUAllocatedMemoryBaseAddress = 0;
 
 
-			for (const TPair<FRHIBuffer*, void*>& Entry : BufferGPUMemoryMap)
+			for (const TPair<FRHIBuffer*, FCudaExternalBufferMapping>& Entry : BufferGPUMemoryMap)
 			{
-				if (Entry.Value)
+
+				if (Entry.Value.ExternalMemory)
 				{
-					CudaModule.DriverAPI()->cuMemFree(reinterpret_cast<CUdeviceptr>(Entry.Value));
+					Status = CudaModule.DriverAPI()->cuDestroyExternalMemory(reinterpret_cast<CUexternalMemory>(Entry.Value.ExternalMemory));
+					if (Status != CUDA_SUCCESS)
+					{
+						UE_LOG(LogRivermax, Error, TEXT("CUDA: Error cleaning up external memory mapping. Error: %d"), Status);
+					}
 				}
 			}
 			BufferGPUMemoryMap.Empty();
@@ -727,9 +787,6 @@ namespace UE::RivermaxCore::Private
 				UE_LOG(LogRivermax, Warning, TEXT("Failed to destroy cuda stream. Status: %d"), Status);
 			}
 			GPUStream = nullptr;
-
-
-			CudaModule.DriverAPI()->cuCtxPopCurrent(nullptr);
 		}
 	}
 
@@ -777,7 +834,7 @@ namespace UE::RivermaxCore::Private
 
 			FCUDAModule& CudaModule = FModuleManager::GetModuleChecked<FCUDAModule>("CUDA");
 
-			CudaModule.DriverAPI()->cuCtxPushCurrent(CudaModule.GetCudaContext());
+			FRivermaxScopedCudaContext CudaScopedContext(CudaModule.GetCudaContextForDevice(GPUDeviceIndex));
 
 			{
 				TRACE_CPUPROFILER_EVENT_SCOPE(Rmax_CudaImportMemory);
@@ -787,6 +844,7 @@ namespace UE::RivermaxCore::Private
 				if (D3D12BufferHandle)
 				{
 					CloseHandle(D3D12BufferHandle);
+					D3D12BufferHandle = nullptr;
 				}
 
 				if (Result != CUDA_SUCCESS)
@@ -802,23 +860,25 @@ namespace UE::RivermaxCore::Private
 				CUDA_EXTERNAL_MEMORY_BUFFER_DESC BufferDescription = {};
 				BufferDescription.offset = 0;
 				BufferDescription.size = BufferMemorySize;
-				CUdeviceptr NewMemory;
+				CUdeviceptr NewMemory = 0;
 				const CUresult Result = FCUDAModule::CUDA().cuExternalMemoryGetMappedBuffer(&NewMemory, MappedExternalMemory, &BufferDescription);
 				if (Result != CUDA_SUCCESS || NewMemory == 0)
 				{
 					UE_LOG(LogRivermax, Error, TEXT("Failed to get shared buffer mapped memory. Error: %d"), Result);
-					CudaModule.DriverAPI()->cuCtxPushCurrent(nullptr);
+					FCUDAModule::CUDA().cuDestroyExternalMemory(MappedExternalMemory);
 					return nullptr;
 				}
+				FCudaExternalBufferMapping MappingInfo;
+				MappingInfo.ExternalMemory = MappedExternalMemory;
+				MappingInfo.MappedPtr = reinterpret_cast<void*>(NewMemory);
 
-				BufferGPUMemoryMap.Add(InBuffer, reinterpret_cast<void*>(NewMemory));
+				BufferGPUMemoryMap.Add(InBuffer, MappingInfo);
 			}
 
-			CudaModule.DriverAPI()->cuCtxPushCurrent(nullptr);
 		}
 
 		// At this point, we have the mapped buffer in cuda space and we can use it to schedule a memcpy on cuda engine.
-		return BufferGPUMemoryMap[InBuffer];
+		return BufferGPUMemoryMap[InBuffer].MappedPtr;
 	}
 
 	void FRivermaxInputStream::ProcessSRD(const FRTPHeader& RTPHeader, const uint8* DataPtr)
@@ -893,14 +953,19 @@ namespace UE::RivermaxCore::Private
 				// For gpudirect, we need all payload sizes to be equal with the exception of the last one.
 				if (FrameDescriptionTracking.CommonPayloadSize <= 0)
 				{
-					UE_LOG(LogRivermax, Warning, TEXT("Unsupported variable SRD length detected while GPUDirect for input stream is used. Disable and reopen the stream."));
+					UE_LOG(LogRivermax, Warning, TEXT("Unsupported variable SRD length detected while GPUDirect is used. Disable GPUDirect and reopen the stream."));
 					Listener->OnStreamError();
-					bIsShuttingDown = true;
+					bIsShuttingDown.store(true, std::memory_order_release);
+
+					Listener->OnVideoFrameReceptionError(StaticCastSharedPtr<IRivermaxVideoSample>(InSample));
+					State = EReceptionState::FrameError;
+					FrameErrorState(RTPHeader);
+					return;
 				}
 
 				//Frame received entirely, time to copy it from rivermax gpu scratchpad to our own gpu memory
 				FCUDAModule& CudaModule = FModuleManager::GetModuleChecked<FCUDAModule>("CUDA");
-				CUresult Result = CudaModule.DriverAPI()->cuCtxPushCurrent(CudaModule.GetCudaContext());
+				FRivermaxScopedCudaContext CudaScopedContext(CudaModule.GetCudaContextForDevice(GPUDeviceIndex));
 
 				const CUdeviceptr DestinationGPUMemory = reinterpret_cast<CUdeviceptr>(StreamData.CurrentFrameVideoBuffer);
 				const CUdeviceptr SourceGPUMemoryOne = reinterpret_cast<CUdeviceptr>(StreamData.DeviceWritePointerOne);
@@ -923,7 +988,7 @@ namespace UE::RivermaxCore::Private
 				StrideDescription.dstPitch = FrameDescriptionTracking.CommonPayloadSize; //Destination pitch is the fixed SRD size we received
 				StrideDescription.WidthInBytes = FrameDescriptionTracking.CommonPayloadSize; //Width in bytes is the amount to copy, the SRD size
 				StrideDescription.Height = NumSRDPartOne;
-				Result = CudaModule.DriverAPI()->cuMemcpy2DAsync(&StrideDescription, reinterpret_cast<CUstream>(GPUStream));
+				CUresult Result = CudaModule.DriverAPI()->cuMemcpy2DAsync(&StrideDescription, reinterpret_cast<CUstream>(GPUStream));
 
 				// Keep track where we are at in the source buffer
 				CUdeviceptr SourcePtr = SourceGPUMemoryOne + (ExpectedPayloadSize * NumSRDPartOne);
@@ -943,11 +1008,24 @@ namespace UE::RivermaxCore::Private
 				const uint32 TotalMemoryCopied = (FrameDescriptionTracking.CommonPayloadSize * (NumSRDPartOne + NumSRDPartTwo));
 				if (TotalMemoryCopied < StreamData.ExpectedSize)
 				{
+					const uint32 LastPacketBytes = RTPHeader.SRD1.bHasContinuation ? (RTPHeader.SRD1.Length + RTPHeader.SRD2.Length) : RTPHeader.SRD1.Length;
+
 					const uint32 SizeLeftToCopy = StreamData.ExpectedSize - TotalMemoryCopied;
-					if (ensure(SizeLeftToCopy == RTPHeader.SRD1.Length))
+					if (SizeLeftToCopy > 0)
 					{
-						CUdeviceptr DestinationPtr = DestinationGPUMemory + TotalMemoryCopied;
-						Result = CudaModule.DriverAPI()->cuMemcpyDtoDAsync(DestinationPtr, SourcePtr, SizeLeftToCopy, reinterpret_cast<CUstream>(GPUStream));
+						if (ensure(SizeLeftToCopy == LastPacketBytes))
+						{
+							CUdeviceptr DestinationPtr = DestinationGPUMemory + TotalMemoryCopied;
+							Result = CudaModule.DriverAPI()->cuMemcpyDtoDAsync(DestinationPtr, SourcePtr, SizeLeftToCopy, reinterpret_cast<CUstream>(GPUStream));
+						}
+						else
+						{
+							UE_LOG(LogRivermax, Warning, TEXT("Unexpected tail copy size: left=%u lastPacket=%u"), SizeLeftToCopy, LastPacketBytes);
+							Listener->OnVideoFrameReceptionError(StaticCastSharedPtr<IRivermaxVideoSample>(InSample));
+							State = EReceptionState::FrameError;
+							FrameErrorState(RTPHeader);
+							return;
+						}
 					}
 				}
 
@@ -960,41 +1038,49 @@ namespace UE::RivermaxCore::Private
 					return;
 				}
 
-				auto CudaCallback = [](void* userData)
+				auto CudaCallback = [](void* UserData)
 				{
 					TRACE_CPUPROFILER_EVENT_SCOPE(FRivermaxInputStream::MemcopyCallback);
-					if (userData)
+					if (UserData)
 					{
 						// It might happen that our stream has been closed once the callback is triggered
-						TWeakPtr<FCallbackPayload>* WeakPayloadPtr = reinterpret_cast<TWeakPtr<FCallbackPayload>*>(userData);
-						TWeakPtr<FCallbackPayload> WeakPayload = *WeakPayloadPtr;
-						if (TSharedPtr<FCallbackPayload> Payload = WeakPayload.Pin())
+						TUniquePtr<TWeakPtr<FCallbackPayload>> WeakPtrHolder(static_cast<TWeakPtr<FCallbackPayload>*>(UserData));
+						if (TSharedPtr<FCallbackPayload> Payload = WeakPtrHolder->Pin())
 						{
-							Payload->bIsWaitingForPendingCopy = false;
+							Payload->bIsWaitingForPendingCopy.store(false, std::memory_order_release);
 						}
 					}
 				};
 
 				// Schedule a callback to know when to make the frame available
-				CallbackPayload->bIsWaitingForPendingCopy = true;
-				TWeakPtr<FCallbackPayload> WeakPayload = CallbackPayload;
-				CudaModule.DriverAPI()->cuLaunchHostFunc(reinterpret_cast<CUstream>(GPUStream), CudaCallback, &WeakPayload);
+				CallbackPayload->bIsWaitingForPendingCopy.store(true, std::memory_order_release);
+				TWeakPtr<FCallbackPayload>* WeakPayloadPtr = new TWeakPtr<FCallbackPayload>(CallbackPayload);
+				Result = CudaModule.DriverAPI()->cuLaunchHostFunc(reinterpret_cast<CUstream>(GPUStream), CudaCallback, WeakPayloadPtr);
+				if (Result != CUDA_SUCCESS)
+				{
+					delete WeakPayloadPtr;
 
-				FCUDAModule::CUDA().cuCtxPopCurrent(nullptr);
+					UE_LOG(LogRivermax, Error, TEXT("Failed to process the last packet. cuLaunchHostFunc failed. Error: %d"), Result);
+					Listener->OnVideoFrameReceptionError(StaticCastSharedPtr<IRivermaxVideoSample>(InSample));
+					State = EReceptionState::FrameError;
+					FrameErrorState(RTPHeader);
 
+					CallbackPayload->bIsWaitingForPendingCopy.store(false, std::memory_order_release);
+					return;
+				}
 				// For now, we wait for the cuda callback before we move on receiving next frame. 
 				// Will need to update this and make the frame available from the cuda callback to avoid losing packets 
 				{
 					TRACE_CPUPROFILER_EVENT_SCOPE(FRivermaxInputStream:WaitingPendingOperation)
 					const double CallbackTimestamp = FPlatformTime::Seconds();
-					while (CallbackPayload->bIsWaitingForPendingCopy == true && bIsShuttingDown == false)
+					while (CallbackPayload->bIsWaitingForPendingCopy.load(std::memory_order_acquire) && bIsShuttingDown.load(std::memory_order_acquire) == false)
 					{
 						FPlatformProcess::SleepNoStats(UE::RivermaxCore::Private::Utils::SleepTimeSeconds);
 						if (FPlatformTime::Seconds() - CallbackTimestamp > CVarWaitForCompletionTimeout.GetValueOnAnyThread())
 						{
 							UE_LOG(LogRivermax, Error, TEXT("Waiting for gpu copy of sample timed out."));
 							Listener->OnStreamError();
-							CallbackPayload->bIsWaitingForPendingCopy = false;
+							CallbackPayload->bIsWaitingForPendingCopy.store(false, std::memory_order_release);
 							break;
 						}
 					}
@@ -1005,7 +1091,7 @@ namespace UE::RivermaxCore::Private
 			FrameDescriptionTracking.ResetSingleFrameTracking();
 
 			// No need to provide the new frame and prepare the next one if we are shutting down
-			if (bIsShuttingDown == false)
+			if (bIsShuttingDown.load(std::memory_order_acquire) == false)
 			{
 				FRivermaxInputVideoFrameReception NewFrame;
 				NewFrame.VideoBuffer = reinterpret_cast<uint8*>(StreamData.CurrentFrameVideoBuffer);
diff --git a/Engine/Plugins/VirtualProduction/Rivermax/RivermaxCore/Source/RivermaxCore/Private/Streams/RivermaxInputStream.h b/Engine/Plugins/VirtualProduction/Rivermax/RivermaxCore/Source/RivermaxCore/Private/Streams/RivermaxInputStream.h
--- a/Engine/Plugins/VirtualProduction/Rivermax/RivermaxCore/Source/RivermaxCore/Private/Streams/RivermaxInputStream.h
+++ b/Engine/Plugins/VirtualProduction/Rivermax/RivermaxCore/Source/RivermaxCore/Private/Streams/RivermaxInputStream.h
@@ -281,13 +281,23 @@ namespace UE::RivermaxCore::Private
 		/** Total allocated gpu memory. */
 		int32 GPUAllocatedMemorySize = 0;
 
+		/** Used to store all the required information for mapping GPU buffers */
+		struct FCudaExternalBufferMapping
+		{
+			// External memory handle created from RHI buffer to be mapped with cuda.
+			void* ExternalMemory = nullptr;
+
+			// Pointer (CUdeviceptr) to mapped external memory.
+			void* MappedPtr = nullptr;
+		};
+
 		/** Map between buffer we are writing to and their mapped address in cuda space */
-		TMap<FRHIBuffer*, void*> BufferGPUMemoryMap;
+		TMap<FRHIBuffer*, FCudaExternalBufferMapping> BufferGPUMemoryMap;
 
 		/** Used to wait for a frame to be copied from gpu to gpu memory */
 		struct FCallbackPayload
 		{
-			volatile bool bIsWaitingForPendingCopy = false;
+			std::atomic<bool> bIsWaitingForPendingCopy = false;
 		};
 		TSharedPtr<FCallbackPayload> CallbackPayload;
 
