Rivermax: Added a bunch of sanity checks for the multi-srd streams.
- Fixed an issue with forcing OneSRD on the last packet even if it still crosses a line.

#rn Rivermax: Fixed an issue with forcing OneSRD on the last packet even if it still crosses a line.

#Synced-CL 0

diff --git a/Engine/Plugins/VirtualProduction/Rivermax/RivermaxCore/Source/RivermaxCore/Private/Streams/RivermaxOutVideoStream.cpp b/Engine/Plugins/VirtualProduction/Rivermax/RivermaxCore/Source/RivermaxCore/Private/Streams/RivermaxOutVideoStream.cpp
--- a/Engine/Plugins/VirtualProduction/Rivermax/RivermaxCore/Source/RivermaxCore/Private/Streams/RivermaxOutVideoStream.cpp
+++ b/Engine/Plugins/VirtualProduction/Rivermax/RivermaxCore/Source/RivermaxCore/Private/Streams/RivermaxOutVideoStream.cpp
@@ -172,12 +172,29 @@ namespace UE::RivermaxCore::Private
 
 			const uint16 SRD1Length = FMath::Min(LineSize - LineSizeOffset, CurrentPayloadSize);
 			const uint16 SRD1PixelCount = SRD1Length / InFormatInfo.PixelGroupSize * InFormatInfo.PixelGroupCoverage;
-			uint16 SRD2Length = SRD1Length < CurrentPayloadSize ? CurrentPayloadSize - SRD1Length : 0;
-			if (SRD2Length && OutLineNumber == ((uint32)AlignedResolution.Y - 1))
-			{
-				SRD2Length = 0;
-			}
 
+			uint16 SRD2Length = (SRD1Length < CurrentPayloadSize) ? (CurrentPayloadSize - SRD1Length) : 0;
+
+			ensureMsgf(
+				((uint32)SRD1Length + (uint32)SRD2Length) == CurrentPayloadSize,
+				TEXT("SRD lengths must sum to payload size. SRD1=%u SRD2=%u Payload=%u Line=%u Offset=%u"), SRD1Length, SRD2Length, CurrentPayloadSize, OutLineNumber, OutSRDOffset
+			);
+
+			ensureMsgf(
+				(SRD1Length % InFormatInfo.PixelGroupSize) == 0 && (SRD2Length % InFormatInfo.PixelGroupSize) == 0,
+				TEXT("SRD lengths must be pixelgroup aligned. SRD1=%u SRD2=%u PixelGroupSize=%u"), SRD1Length, SRD2Length, InFormatInfo.PixelGroupSize
+			);
+
+			// If SRD2 is needed on the last line, something upstream is wrong (payload sizing/header sizing mismatch).
+			ensureMsgf(
+				!(SRD2Length && OutLineNumber == ((uint32)AlignedResolution.Y - 1)),
+				TEXT("Invalid packetization. Payload crosses past last line. SRD2=%u Line=%u Res=%dx%d Offset=%u Payload=%u"),
+				SRD2Length,
+				OutLineNumber,
+				AlignedResolution.X, AlignedResolution.Y,
+				OutSRDOffset,
+				CurrentPayloadSize
+			);
 
 			constexpr uint32 FieldIdentification = 0; //todo when fields are sent for interlace
 
@@ -185,8 +202,8 @@ namespace UE::RivermaxCore::Private
 			VideoHeader->SRD1.Length = ByteSwap(SRD1Length);
 			VideoHeader->SRD1.SetRowNumber(OutLineNumber); //todo divide by 2 if interlaced
 			VideoHeader->SRD1.SetOffset(OutSRDOffset);
-			VideoHeader->SRD1.ContinuationBit = SRD2Length > 0;
-			VideoHeader->SRD1.FieldIdentification = 0; //todo when fields are sent for interlace
+			VideoHeader->SRD1.ContinuationBit = (SRD2Length > 0);
+			VideoHeader->SRD1.FieldIdentification = FieldIdentification;
 #endif
 
 #if RIVERMAX_PACKET_DEBUG
@@ -227,7 +244,7 @@ namespace UE::RivermaxCore::Private
 				VideoHeader->SRD2.SetRowNumber(OutLineNumber);
 				VideoHeader->SRD2.SetOffset(OutSRDOffset);
 				VideoHeader->SRD2.ContinuationBit = 0;
-				VideoHeader->SRD2.FieldIdentification = 0;
+				VideoHeader->SRD2.FieldIdentification = FieldIdentification;
 #endif
 
 				const uint16 SRD2PixelCount = SRD2Length / InFormatInfo.PixelGroupSize * InFormatInfo.PixelGroupCoverage;
@@ -238,6 +255,11 @@ namespace UE::RivermaxCore::Private
 					++OutLineNumber;
 				}
 			}
+			// Never happened. Sanity check.
+			ensureMsgf(
+				OutLineNumber <= (uint32)AlignedResolution.Y,
+				TEXT("LineNumber overflow. Line=%u ResY=%d"), OutLineNumber, AlignedResolution.Y
+			);
 		}
 
 
@@ -586,6 +608,17 @@ namespace UE::RivermaxCore::Private
 		const uint32 PixelCount = StreamOptions->AlignedResolution.X * StreamOptions->AlignedResolution.Y;
 		const uint64 FrameSizeInBytes = PixelCount / FormatInfo.PixelGroupCoverage * FormatInfo.PixelGroupSize;
 
+		// Sanity check. By this point, frame size must be pixel group aligned.
+		ensureMsgf(
+			(FrameSizeInBytes % FormatInfo.PixelGroupSize) == 0,
+			TEXT("FrameSizeInBytes (%llu) must be pixel-group aligned (PixelGroupSize=%u). Res=%dx%d Sampling=%s"),
+			FrameSizeInBytes,
+			FormatInfo.PixelGroupSize,
+			StreamOptions->AlignedResolution.X,
+			StreamOptions->AlignedResolution.Y,
+			LexToString(FormatInfo.Sampling)
+		);
+
 		StreamMemory.PixelGroupPerPacket = StreamMemory.PayloadSize / FormatInfo.PixelGroupSize;
 		StreamMemory.PixelsPerPacket = StreamMemory.PixelGroupPerPacket * FormatInfo.PixelGroupCoverage;
 
@@ -687,58 +720,84 @@ namespace UE::RivermaxCore::Private
 			HeaderMemory->mkey[0] = InvalidKey;
 			HeaderMemory->mkey[1] = InvalidKey;
 		}
+
 		uint64 ProcessedBytes = 0;
-		uint64 LineSize = 0;
-		for (int32 PayloadSizeIndex = 0; PayloadSizeIndex < RealPacketsPerFrame; ++PayloadSizeIndex)
+		uint32 LineSize = 0; // bytes consumed in current line (0..BytesPerRow)
+
+		const uint32 BytesPerRowU32 = (uint32)BytesPerRow;
+
+		for (uint32 PacketIndex = 0; PacketIndex < RealPacketsPerFrame; ++PacketIndex)
 		{
-			uint32 HeaderSize = FVideoRTPHeader::OneSRDSize;
-			uint32 ThisPayloadSize = StreamMemory.PayloadSize;
+			uint32 HeaderSize = 0;
+			uint32 ThisPayloadSize = 0;
+
 			if (ProcessedBytes < FrameSizeInBytes)
 			{
-				if ((LineSize + StreamMemory.PayloadSize) == BytesPerRow)
+				// 1. Decide actual payload size for this packet (full or last remainder)
+				ThisPayloadSize = StreamMemory.PayloadSize;
+				if ((ProcessedBytes + ThisPayloadSize) > FrameSizeInBytes)
 				{
-					LineSize = 0;
+					ThisPayloadSize = (uint32)(FrameSizeInBytes - ProcessedBytes);
 				}
-				else if ((LineSize + StreamMemory.PayloadSize) > BytesPerRow)
+
+				// Payload must always align to pixel groups and be a multiple of header defined payload size. 
+				ensureMsgf(
+					(ThisPayloadSize % FormatInfo.PixelGroupSize) == 0,
+					TEXT("Payload size (%u) must be pixel-group aligned (PixelGroupSize=%u). Packet=%u ProcessedBytes=%llu FrameSizeInBytes=%llu"),
+					ThisPayloadSize,
+					FormatInfo.PixelGroupSize,
+					PacketIndex,
+					ProcessedBytes,
+					FrameSizeInBytes
+				);
+
+				// 2. Decide SRD count using THIS payload size
+				const uint32 RemainingInLine = BytesPerRowU32 - LineSize;
+
+				if (ThisPayloadSize <= RemainingInLine)
 				{
-					HeaderSize = FVideoRTPHeader::TwoSRDSize;
-					LineSize = StreamMemory.PayloadSize - (BytesPerRow - LineSize);
-					if (LineSize > BytesPerRow)
+					HeaderSize = FVideoRTPHeader::OneSRDSize;
+
+					LineSize += ThisPayloadSize;
+					if (LineSize == BytesPerRowU32)
 					{
-						UE_LOG(LogRivermax, Warning, TEXT("Unsupported small resolution, %dx%d, needing more than 2 SRD to express"), StreamOptions->AlignedResolution.X, StreamOptions->AlignedResolution.Y);
-						return false;
+						LineSize = 0;
 					}
 				}
 				else
 				{
-					// Keep track of line size offset to know when to use TwoSRDs
-					LineSize += StreamMemory.PayloadSize;
-				}
+					HeaderSize = FVideoRTPHeader::TwoSRDSize;
 
-				// This means that this is the header for the last packet and it is smaller than all the previous ones.
-				if ((ProcessedBytes + StreamMemory.PayloadSize) > FrameSizeInBytes)
-				{
-					HeaderSize = FVideoRTPHeader::OneSRDSize;
-					ThisPayloadSize = FrameSizeInBytes - ProcessedBytes;
+					const uint32 Spill = ThisPayloadSize - RemainingInLine;
+					LineSize = Spill;
+
+					// Would require more than two SRDs, which is unsupported
+					if (LineSize > BytesPerRowU32)
+					{
+						UE_LOG(LogRivermax, Warning,
+							TEXT("Unsupported resolution/payload: packet crosses more than 2 lines (needs >2 SRD). Res=%dx%d Spill=%u BytesPerRow=%u Payload=%u"),
+							StreamOptions->AlignedResolution.X, StreamOptions->AlignedResolution.Y, LineSize, BytesPerRowU32, ThisPayloadSize);
+						return false;
+					}
 				}
 			}
 			else
 			{
-				// Extra header/payload required for the chunk alignment are set to 0. Nothing has to be sent out the wire.
+				// Extra packet slots for chunk alignment.
 				HeaderSize = 0;
 				ThisPayloadSize = 0;
 			}
 
-			// All buffers are configured the same so compute header and payload sizes once and assigned to all impacted locations
+			// All buffers configured the same
 			for (uint32 BufferIndex = 0; BufferIndex < StreamMemory.FramesFieldPerMemoryBlock; ++BufferIndex)
 			{
-				StreamMemory.HeaderSizes[PayloadSizeIndex + (BufferIndex * RealPacketsPerFrame)] = HeaderSize;
-				StreamMemory.PayloadSizes[PayloadSizeIndex + (BufferIndex * RealPacketsPerFrame)] = ThisPayloadSize;
+				StreamMemory.HeaderSizes[PacketIndex + (BufferIndex * RealPacketsPerFrame)] = HeaderSize;
+				StreamMemory.PayloadSizes[PacketIndex + (BufferIndex * RealPacketsPerFrame)] = ThisPayloadSize;
 			}
 
 			if (CachedCVars.bPrefillRTPHeaders && HeaderSize > 0)
 			{
-				RTPFiller.Update(PayloadSizeIndex);
+				RTPFiller.Update(PacketIndex);
 			}
 
 			ProcessedBytes += ThisPayloadSize;
