Rivermax: Migrated Rivermax Output specific property PresentationQueueSize to the general Media Output property NumberOfTextureBuffers.
- Other Minor improvements.

#rn Migrated Rivermax Output specific property PresentationQueueSize to the general Media Output property NumberOfTextureBuffers.

#Synced-CL 0

diff --git a/Engine/Plugins/VirtualProduction/Rivermax/RivermaxCore/Source/RivermaxCore/Private/Streams/RivermaxOutStream.cpp b/Engine/Plugins/VirtualProduction/Rivermax/RivermaxCore/Source/RivermaxCore/Private/Streams/RivermaxOutStream.cpp
--- a/Engine/Plugins/VirtualProduction/Rivermax/RivermaxCore/Source/RivermaxCore/Private/Streams/RivermaxOutStream.cpp
+++ b/Engine/Plugins/VirtualProduction/Rivermax/RivermaxCore/Source/RivermaxCore/Private/Streams/RivermaxOutStream.cpp
@@ -1403,19 +1403,19 @@ namespace UE::RivermaxCore::Private
 			int64 TimeDeltaAlignmentPoint = (static_cast<int64>(TimeStamp) - static_cast<int64>(FrameScheduledTimeNs))/1000;
 			int64 TimeDeltaExpectedChunkCompletionTime = (static_cast<int64>(TimeStamp) - static_cast<int64>(TrackedChunk.ExpectedCompletionTime)) / 1000;
 
-			UE_LOG(LogRivermax, VeryVerbose, TEXT("Chunk \"%s\" was completed with timestamp: %u nanoseconds, which is %d microseconds away from the expected alignment point.\n")
-											TEXT("The chunk is %d microseconds away from the predicted chunk completion time.")
+			UE_LOG(LogRivermax, VeryVerbose, TEXT("Chunk \"%s\" was completed with timestamp: %llu nanoseconds, which is %lld microseconds away from the expected alignment point.\n")
+											TEXT("The chunk is %lld microseconds away from the predicted chunk completion time.")
 											, *HumanReadableString, TimeStamp, TimeDeltaAlignmentPoint, TimeDeltaExpectedChunkCompletionTime);
 
 			if (TimeDeltaAlignmentPoint < 0)
 			{
-				UE_LOG(LogRivermax, Error, TEXT("Chunk \"%s\" was completed %d microseconds before the first packet was expected to be on the wire. "), *HumanReadableString, FMath::Abs(TimeDeltaAlignmentPoint));
+				UE_LOG(LogRivermax, Error, TEXT("Chunk \"%s\" was completed %lld microseconds before the first packet was expected to be on the wire. "), *HumanReadableString, FMath::Abs(TimeDeltaAlignmentPoint));
 			}
 
 			// ExpectedCompletionTime is more of a prediction and isn't precise. It also depends on Rivermax internal workings. Therefore not being on time in this case is not an error.
 			if (TimeStamp > TrackedChunk.ExpectedCompletionTime)
 			{
-				UE_LOG(LogRivermax, Verbose, TEXT("Chunk \"%s\" was completed %d microseconds after it was expected to be completed. "), *HumanReadableString, FMath::Abs(TimeDeltaExpectedChunkCompletionTime));
+				UE_LOG(LogRivermax, Verbose, TEXT("Chunk \"%s\" was completed %lld microseconds after it was expected to be completed. "), *HumanReadableString, FMath::Abs(TimeDeltaExpectedChunkCompletionTime));
 			}
 
 			LastTimeStamp = TimeStamp;
diff --git a/Engine/Plugins/VirtualProduction/Rivermax/RivermaxMedia/Source/RivermaxMedia/Private/RivermaxMediaOutput.cpp b/Engine/Plugins/VirtualProduction/Rivermax/RivermaxMedia/Source/RivermaxMedia/Private/RivermaxMediaOutput.cpp
--- a/Engine/Plugins/VirtualProduction/Rivermax/RivermaxMedia/Source/RivermaxMedia/Private/RivermaxMediaOutput.cpp
+++ b/Engine/Plugins/VirtualProduction/Rivermax/RivermaxMedia/Source/RivermaxMedia/Private/RivermaxMediaOutput.cpp
@@ -15,7 +15,8 @@ enum class ERivermaxMediaOutputVersion : int32
 {
 	BeforeAncTimecode = 0,
 	AncTimecode = 1,
-	Latest = AncTimecode
+	FrameManagementChange = 2,
+	Latest = FrameManagementChange
 };
 
 struct FRivermaxMediaOutputCustomVersion
@@ -170,7 +171,7 @@ UE::RivermaxCore::FRivermaxOutputOptions URivermaxMediaOutput::GenerateStreamOpt
 		VideoOptions->Resolution = GetRequestedSize();
 
 		VideoOptions->FrameRate = VideoStream.FrameRate;
-		OutOutputOptions.NumberOfBuffers = PresentationQueueSize;
+		OutOutputOptions.NumberOfBuffers = NumberOfTextureBuffers;
 		VideoOptions->bUseGPUDirect = VideoStream.bUseGPUDirect;
 		OutOutputOptions.AlignmentMode = MediaOutputAlignmentToRivermaxAlignment(AlignmentMode);
 		OutOutputOptions.FrameLockingMode = UE::RivermaxMediaUtils::Private::MediaOutputFrameLockingToRivermax(FrameLockingMode);
@@ -242,8 +243,6 @@ void URivermaxMediaOutput::Serialize(FArchive& Ar)
 void URivermaxMediaOutput::PostLoad()
 {
 	Super::PostLoad();
-	FGuid NewGuid = FGuid::NewGuid();
-	UE_LOG(LogTemp, Log, TEXT("Generated GUID: %s"), *NewGuid.ToString());
 
 	const int32 Version = GetLinkerCustomVersion(FRivermaxMediaOutputCustomVersion::GUID);
 	if (Version < (int32)ERivermaxMediaOutputVersion::AncTimecode)
@@ -259,6 +258,14 @@ void URivermaxMediaOutput::PostLoad()
 		VideoStream.bUseGPUDirect = bUseGPUDirect_DEPRECATED;
 		PRAGMA_ENABLE_DEPRECATION_WARNINGS
 	}
+
+	if (Version < (int32)ERivermaxMediaOutputVersion::FrameManagementChange)
+	{
+		PRAGMA_DISABLE_DEPRECATION_WARNINGS
+		NumberOfTextureBuffers = PresentationQueueSize_DEPRECATED;
+		PRAGMA_ENABLE_DEPRECATION_WARNINGS
+	}
+	
 }
 
 
diff --git a/Engine/Plugins/VirtualProduction/Rivermax/RivermaxMedia/Source/RivermaxMedia/Public/RivermaxMediaOutput.h b/Engine/Plugins/VirtualProduction/Rivermax/RivermaxMedia/Source/RivermaxMedia/Public/RivermaxMediaOutput.h
--- a/Engine/Plugins/VirtualProduction/Rivermax/RivermaxMedia/Source/RivermaxMedia/Public/RivermaxMediaOutput.h
+++ b/Engine/Plugins/VirtualProduction/Rivermax/RivermaxMedia/Source/RivermaxMedia/Public/RivermaxMediaOutput.h
@@ -185,9 +185,10 @@ public:
 	UPROPERTY(EditAnywhere, BlueprintReadWrite, Category = "Output")
 	ERivermaxFrameLockingMode FrameLockingMode = ERivermaxFrameLockingMode::FreeRun;
 
-	/** Number of frames that can be queued / used in output queue. Frame being sent counts for 1. */
-	UPROPERTY(EditAnywhere, BlueprintReadWrite, Category = "Output", meta = (ClampMin = "2", ClampMax = "8", UIMin = "2", UIMax = "8"))
-	int32 PresentationQueueSize = 2;
+	/** DEPRECATED: Formally controlled the number of frames that could be queued / used in output queue. */
+	UE_DEPRECATED(5.8, "PresentationQueueSize is deprecated in favor of NumberOfTextureBuffers from UMediaOutput.")
+	UPROPERTY(meta = (DeprecatedProperty, DeprecationMessage = "Use UMediaOutput::NumberOfTextureBuffers instead."))
+	int32 PresentationQueueSize_DEPRECATED = 2;
 
 	/** 
 	 * Experimental flag to use frame counter instead of using NVIDIA Rivermax clock for timestamping output frames
diff --git a/Engine/Plugins/VirtualProduction/Rivermax/RivermaxMedia/Source/RivermaxMediaEditor/Private/ModularFeatures/RivermaxMediaInitializerFeature.cpp b/Engine/Plugins/VirtualProduction/Rivermax/RivermaxMedia/Source/RivermaxMediaEditor/Private/ModularFeatures/RivermaxMediaInitializerFeature.cpp
--- a/Engine/Plugins/VirtualProduction/Rivermax/RivermaxMedia/Source/RivermaxMediaEditor/Private/ModularFeatures/RivermaxMediaInitializerFeature.cpp
+++ b/Engine/Plugins/VirtualProduction/Rivermax/RivermaxMedia/Source/RivermaxMediaEditor/Private/ModularFeatures/RivermaxMediaInitializerFeature.cpp
@@ -69,7 +69,7 @@ void FRivermaxMediaInitializerFeature::InitializeMediaObjectForTile(UObject* Med
 		RivermaxMediaOutput->AlignmentMode       = ERivermaxMediaAlignmentMode::FrameCreation;
 		RivermaxMediaOutput->bDoContinuousOutput = false;
 		RivermaxMediaOutput->FrameLockingMode    = ERivermaxFrameLockingMode::BlockOnReservation;
-		RivermaxMediaOutput->PresentationQueueSize = 2;
+		RivermaxMediaOutput->NumberOfTextureBuffers = 2;
 		RivermaxMediaOutput->bDoFrameCounterTimestamping = true;
 
 		RivermaxMediaOutput->VideoStream.bOverrideResolution = false;
@@ -101,7 +101,7 @@ void FRivermaxMediaInitializerFeature::InitializeMediaObjectForFullFrame(UObject
 	else if (URivermaxMediaOutput* RivermaxMediaOutput = Cast<URivermaxMediaOutput>(MediaObject))
 	{
 		RivermaxMediaOutput->FrameLockingMode = ERivermaxFrameLockingMode::BlockOnReservation;
-		RivermaxMediaOutput->PresentationQueueSize = 2;
+		RivermaxMediaOutput->NumberOfTextureBuffers = 2;
 		RivermaxMediaOutput->VideoStream.bOverrideResolution   = false;
 		RivermaxMediaOutput->VideoStream.PixelFormat           = ERivermaxMediaOutputPixelFormat::PF_10BIT_RGB;
 		RivermaxMediaOutput->VideoStream.InterfaceAddress      = GetRivermaxInterfaceAddress();
