Rivermax: changes to debugging cvars that throw Rivermax API out of sync.

#rnx

#Synced-CL 0

diff --git a/Engine/Plugins/VirtualProduction/Rivermax/RivermaxCore/Source/RivermaxCore/Private/Streams/RivermaxOutStream.cpp b/Engine/Plugins/VirtualProduction/Rivermax/RivermaxCore/Source/RivermaxCore/Private/Streams/RivermaxOutStream.cpp
--- a/Engine/Plugins/VirtualProduction/Rivermax/RivermaxCore/Source/RivermaxCore/Private/Streams/RivermaxOutStream.cpp
+++ b/Engine/Plugins/VirtualProduction/Rivermax/RivermaxCore/Source/RivermaxCore/Private/Streams/RivermaxOutStream.cpp
@@ -59,9 +59,26 @@ namespace UE::RivermaxCore::Private
 		TEXT("Configures Rivermax stream to use a single memblock potentially improving SDK performance."),
 		ECVF_Default);
 
-	static TAutoConsoleVariable<int32> CVarRivermaxOutputRandomDelay(
-		TEXT("Rivermax.Output.TriggerRandomDelay"), 0,
-		TEXT("Will cause a delay of variable amount of time when next frame is sent."),
+	static TAutoConsoleVariable<float> CVarRivermaxOutputRandomDelay(
+		TEXT("Rivermax.Output.TriggerRandomDelayMs"), 0.0,
+		TEXT("Will cause a delay of selected amount of time when next frame is sent in milliseconds."),
+		ECVF_Default);
+
+	enum class ETimingIssueMode : int32
+	{
+		Disabled = 0,
+		Immediate = 1,
+		Delay = 2,
+		RandomContinuous = 3,
+	};
+
+	static TAutoConsoleVariable<int32> CVarRivermaxOutputTriggerTimingIssue(
+		TEXT("Rivermax.Output.TriggerTimingIssue"), 0,
+		TEXT("This is a debugging Cvar to introduce artificial timing issue.\n\
+		0: disabled \n\
+		1: set to this value to do it once, \n\
+		2: set to this value to skip one frame in rougly 10 seconds.\n\
+		3: set to this value to continuosly and randomly skip 10% of the frames."),
 		ECVF_Default);
 
 	static TAutoConsoleVariable<int32> CVarRivermaxOutputEnableTimingProtection(
@@ -90,11 +107,7 @@ namespace UE::RivermaxCore::Private
 		The value indicates in percent of frame time how much earlier the Rivermax plugin will attempt to commit chunks."),
 		ECVF_Default);
 
-	static bool GbTriggerRandomTimingIssue = false;
-	FAutoConsoleVariableRef CVarTriggerRandomTimingIssue(
-		TEXT("Rivermax.Sync.TriggerRandomTimingIssue")
-		, UE::RivermaxCore::Private::GbTriggerRandomTimingIssue
-		, TEXT("Randomly triggers a timing issue to test self repair."), ECVF_Cheat);
+
 
 	FRivermaxOutStream::FRivermaxOutStream(const TArray<char>& InSDPDescription)
 		: bIsActive(false)
@@ -127,6 +140,11 @@ namespace UE::RivermaxCore::Private
 			CvarDelay->OnChangedDelegate().AddRaw(this, &FRivermaxOutStream::OnCVarRandomDelayChanged);
 		}
 
+		if (IConsoleVariable* CvarDelay = CVarRivermaxOutputTriggerTimingIssue.AsVariable())
+		{
+			CvarDelay->OnChangedDelegate().AddRaw(this, &FRivermaxOutStream::OnCVarTriggerTimingIssueChanged);
+		}
+
 		CacheCVarValues();
 		
 		// Cache API entry point
@@ -269,6 +287,11 @@ namespace UE::RivermaxCore::Private
 		{
 			CvarDelay->OnChangedDelegate().RemoveAll(this);
 		}
+
+		if (IConsoleVariable* CvarDelay = CVarRivermaxOutputTriggerTimingIssue.AsVariable())
+		{
+			CvarDelay->OnChangedDelegate().RemoveAll(this);
+		}
 	}
 
 	void FRivermaxOutStream::InitializeNextFrame(const TSharedPtr<FRivermaxOutputFrame>& NextFrame)
@@ -1172,6 +1195,20 @@ namespace UE::RivermaxCore::Private
 		bTriggerRandomDelay = true;
 	}
 
+	void FRivermaxOutStream::OnCVarTriggerTimingIssueChanged(IConsoleVariable* Var)
+	{
+		ScheduledTimeForArtificialTimingIssue = 0.0;
+		ETimingIssueMode State = static_cast<ETimingIssueMode>(Var->GetInt());
+		if (State == ETimingIssueMode::Delay)
+		{
+			ScheduledTimeForArtificialTimingIssue = FPlatformTime::Seconds() + 10.;
+		}
+		else if (State == ETimingIssueMode::Immediate)
+		{
+			ScheduledTimeForArtificialTimingIssue = FPlatformTime::Seconds();
+		}
+	}
+
 	void FRivermaxOutStream::CalculateFrameTimestamp()
 	{
 		using namespace UE::RivermaxCore::Private::Utils;
@@ -1230,7 +1267,7 @@ namespace UE::RivermaxCore::Private
 		{
 			TRACE_CPUPROFILER_EVENT_SCOPE(FRivermaxOutStream::RandomDelay);
 			bTriggerRandomDelay = false;
-			FPlatformProcess::SleepNoStats(FMath::RandRange(2e-3, 4e-3));
+			FPlatformProcess::SleepNoStats(CVarRivermaxOutputRandomDelay.GetValueOnRenderThread()/1000.);
 		}
 
 		// Calculate frame's timestamp only once and reuse in RTP build
@@ -1245,6 +1282,8 @@ namespace UE::RivermaxCore::Private
 		// Process completions of the previous frame's chunks.
 		Stats.ChunkCompletionTracker->PollAndReportCompletionOfTrackedChunks();
 
+		ETimingIssueMode TriggerTimingIssue = static_cast<ETimingIssueMode>(CVarRivermaxOutputTriggerTimingIssue.GetValueOnAnyThread());
+
 		do
 		{
 			if (bIsActive)
@@ -1273,17 +1312,16 @@ namespace UE::RivermaxCore::Private
 					const bool bIsChunkOnTime = IsChunkOnTime();
 					CurrentFrame->bCaughtTimingIssue = !bIsChunkOnTime;
 
-					if (UE::RivermaxCore::Private::GbTriggerRandomTimingIssue)
+					if ((TriggerTimingIssue == ETimingIssueMode::Immediate ||
+						TriggerTimingIssue == ETimingIssueMode::Delay)
+						&& (FPlatformTime::Seconds() > ScheduledTimeForArtificialTimingIssue && ScheduledTimeForArtificialTimingIssue > 0.0))
 					{
-						FRandomStream RandomStream(FPlatformTime::Cycles64());
-						const bool bTriggerDesync = (RandomStream.FRandRange(0.0, 1.0) > 0.7) ? true : false;
-						if (bTriggerDesync)
-						{
-							TRACE_CPUPROFILER_EVENT_SCOPE(RmaxOut::ForceTimingIssue);
-							CurrentFrame->bCaughtTimingIssue = true;
-						}
-
-						UE::RivermaxCore::Private::GbTriggerRandomTimingIssue = false;
+						CurrentFrame->bCaughtTimingIssue = true;
+						ScheduledTimeForArtificialTimingIssue = 0.0;
+					}
+					else if (TriggerTimingIssue == ETimingIssueMode::RandomContinuous && FMath::FRandRange(0., 1.) < .1)
+					{
+						CurrentFrame->bCaughtTimingIssue = true;
 					}
 				}
 
diff --git a/Engine/Plugins/VirtualProduction/Rivermax/RivermaxCore/Source/RivermaxCore/Private/Streams/RivermaxOutStream.h b/Engine/Plugins/VirtualProduction/Rivermax/RivermaxCore/Source/RivermaxCore/Private/Streams/RivermaxOutStream.h
--- a/Engine/Plugins/VirtualProduction/Rivermax/RivermaxCore/Source/RivermaxCore/Private/Streams/RivermaxOutStream.h
+++ b/Engine/Plugins/VirtualProduction/Rivermax/RivermaxCore/Source/RivermaxCore/Private/Streams/RivermaxOutStream.h
@@ -442,6 +442,9 @@ namespace UE::RivermaxCore::Private
 		/** Called when delay request cvar has been changed */
 		void OnCVarRandomDelayChanged(IConsoleVariable* Var);
 
+		/** Called when Trigger Timing issue cvar value has been changed */
+		void OnCVarTriggerTimingIssueChanged(IConsoleVariable* Var);
+
 		/** Update frame's timestamp to be used when setting every RTP headers */
 		void CalculateFrameTimestamp();
 
@@ -528,6 +531,10 @@ namespace UE::RivermaxCore::Private
 
 		/** An actual string of the SDP file. */
 		TArray<char> SDPDescription;
+
+	private:
+		/** Used for debugging Rivermax streams when frame needs to be skipped after a certain period. */
+		double ScheduledTimeForArtificialTimingIssue = 0.0;
 	};
 }
 
